Snapshot for changelist 2410693.

[CL 2410693 in Main branch]

diff --git a/Engine/Plugins/Developer/GitSourceControl/.gitignore b/Engine/Plugins/Developer/GitSourceControl/.gitignore
new file mode 100644
--- /dev/null
+++ b/Engine/Plugins/Developer/GitSourceControl/.gitignore
@@ -0,0 +1,3 @@
+/Binaries/*/*.pdb
+/Binaries/*/*Debug*
+/Intermediate
\ No newline at end of file
diff --git a/Engine/Plugins/Developer/GitSourceControl/GitSourceControl.uplugin b/Engine/Plugins/Developer/GitSourceControl/GitSourceControl.uplugin
new file mode 100644
--- /dev/null
+++ b/Engine/Plugins/Developer/GitSourceControl/GitSourceControl.uplugin
@@ -0,0 +1,22 @@
+{
+	"FileVersion" : 3,
+	
+	"FriendlyName" : "Git",
+	"Version" : 12,
+	"VersionName" : "0.6.3",
+	"EngineVersion" : "4.8.0",
+	"Description" : "Git source control management",
+	"Category" : "Editor.Source Control",
+	"CreatedBy" : "SRombauts",
+	"CreatedByURL" : "http://srombauts.github.com",
+	"CanContainContent" : false,
+	"IsBetaVersion" : true,
+	
+	"Modules" :
+	[
+		{
+			"Name" : "GitSourceControl",
+			"Type" : "Editor"
+		}
+	]
+}
diff --git a/Engine/Plugins/Developer/GitSourceControl/README.txt b/Engine/Plugins/Developer/GitSourceControl/README.txt
new file mode 100644
--- /dev/null
+++ b/Engine/Plugins/Developer/GitSourceControl/README.txt
@@ -0,0 +1,90 @@
+Unreal Engine 4 Git Source Control Plugin
+-----------------------------------------
+
+UE4GitPlugin is a simple Git Source Control Plugin for Unreal Engine 4.8
+
+Copyright (c) 2014-2015 [REDACTED_PERSON] ([REDACTED_EMAIL])
+
+### Beta version 0.6.3:
+- display status icons to show modified/added/deleted/untracked files
+- show history of a file
+- diff against depot or between previous versions of a file
+- revert modifications of a file
+- add a file
+- delete a file
+- checkin/commit a file (cannot handle atomically more than 20 files)
+- show current branch name in status text
+- solve a merge conflict
+
+### What *cannot* be done presently (TODO list for v1.0, ordered by priority):
+- merge blueprints
+- initialize a new Git local repository ('git init') to manage your UE4 Game Project.
+- tags: implement ISourceControlLabel to manage git tags
+- .uproject file state si not visible in the current Editor
+- Branch is not in the current Editor workflow (but on Epic Roadmap)
+- Pull/Fetch/Push are not in the current Editor workflow
+- Amend a commit is not in the current Editor workflow
+- configure user name & email ('git config user.name' & git config user.email')
+
+### Known issues:
+- the Editor does not show deleted files (only when deleted externaly?)
+- the Editor does not show missing files
+- missing localisation for git specific messages
+- migrate an asset should add it to the destination project if also under Git (needs management of 'out of tree' files)
+- displaying states of 'Engine' assets (also needs management of 'out of tree' files)
+- issue #22: A Move/Rename leaves a redirector file behind
+- issue #11: Add the "Resolve" operation introduced in Editor 4.3
+- improve the 'Init' window text, hide it if connection is already done, auto connect
+- reverting an asset does not seem to update content in Editor! Issue in Editor?
+- file history show Revision as signed integer instead of hexadecimal SHA1
+- file history does not report file size
+- standard Editor commit dialog ask if user wants to "Keep Files Checked Out" => no use for Git or Mercurial CanCheckOut()==false
+- Windows only (64 bits) -> Mac compiles but needs testing/debugging (Linux source control is not supported by Editor)
+
+### Wishlist (after v1.0):
+- [git-annexe and/or git-media - #1 feature request]([REDACTED_PROFILE_URL]/UE4GitPlugin/issues/1)
+
+### In-code TODO list (internal roadmap):
+
+- FGitConnectWorker::Execute (While project not in Git source control)
+  Improve error message "You should check out a working copy..."
+  => double error message (and in reverse order) with "Project is not part of a Git working copy"
+
+- FGitResolveWorker (GitSourceControlOperations.h)
+  git add to mark a conflict as resolved
+  
+- FGitSourceControlRevision::GetBranchSource() const
+  if this revision was copied from some other revision, then that source revision should
+	be returned here (this should be determined when history is being fetched)
+- FGitSourceControlState::GetBaseRevForMerge()
+  get revision of the merge-base (https://www.kernel.org/pub/software/scm/git/docs/git-merge-base.html)
+	
+- FGitConnectWorker::Execute()
+  popup to propose to initialize the git repository "git init + .gitignore"
+
+- FGitSyncWorker (GitSourceControlOperations.h)
+  git fetch remote(s) to be able to show files not up-to-date with the serveur
+- FGitSourceControlState::IsCurrent() const
+  check the state of the HEAD versus the state of tracked branch on remote
+
+- FGitSourceControlRevision::GetFileSize() const
+	git log does not give us the file size, but we could run a specific command
+
+- GitSourceControlUtils::CheckGitAvailability
+  also check Git config user.name & user.email
+
+Windows:
+- GitSourceControlUtils::FindGitBinaryPath
+  use the Windows registry to find Git
+
+Mac:
+- GitSourceControlUtils::RunCommandInternalRaw
+  Specifying the working copy (the root) of the git repository (before the command itself)
+	does not work in UE4.1 on Mac if there is a space in the path ("[REDACTED_USER_HOME]/Unreal Project/MyProject")
+
+Bug reports?
+- FGitSourceControlRevision::Get
+  Bug report: a counter is needed to avoid overlapping files; temp files are not (never?) released by Editor!
+
+- GitSourceControlUtils::UpdateCachedStates
+  // State->TimeStamp = InState.TimeStamp; // Bug report: Workaround a bug with the Source Control Module not updating file state after a "Save"
diff --git a/Engine/Plugins/Developer/GitSourceControl/Source/GitSourceControl/GitSourceControl.Build.cs b/Engine/Plugins/Developer/GitSourceControl/Source/GitSourceControl/GitSourceControl.Build.cs
new file mode 100644
--- /dev/null
+++ b/Engine/Plugins/Developer/GitSourceControl/Source/GitSourceControl/GitSourceControl.Build.cs
@@ -0,0 +1,19 @@
+// Copyright 1998-2015 Epic Games, Inc. All Rights Reserved.
+
+using UnrealBuildTool;
+
+public class GitSourceControl : ModuleRules
+{
+	public GitSourceControl(TargetInfo Target)
+	{
+		PrivateDependencyModuleNames.AddRange(
+			new string[] {
+				"Core",
+				"Slate",
+				"SlateCore",
+				"EditorStyle",
+				"SourceControl",
+			}
+		);
+	}
+}
diff --git a/Engine/Plugins/Developer/GitSourceControl/Source/GitSourceControl/Private/GitSourceControlCommand.cpp b/Engine/Plugins/Developer/GitSourceControl/Source/GitSourceControl/Private/GitSourceControlCommand.cpp
new file mode 100644
--- /dev/null
+++ b/Engine/Plugins/Developer/GitSourceControl/Source/GitSourceControl/Private/GitSourceControlCommand.cpp
@@ -0,0 +1,44 @@
+// Copyright 1998-2015 Epic Games, Inc. All Rights Reserved.
+
+#include "GitSourceControlPrivatePCH.h"
+#include "GitSourceControlCommand.h"
+#include "GitSourceControlModule.h"
+#include "GitSourceControlProvider.h"
+#include "IGitSourceControlWorker.h"
+#include "SGitSourceControlSettings.h"
+
+FGitSourceControlCommand::FGitSourceControlCommand(const TSharedRef<class ISourceControlOperation, ESPMode::ThreadSafe>& InOperation, const TSharedRef<class IGitSourceControlWorker, ESPMode::ThreadSafe>& InWorker, const FSourceControlOperationComplete& InOperationCompleteDelegate)
+	: Operation(InOperation)
+	, Worker(InWorker)
+	, OperationCompleteDelegate(InOperationCompleteDelegate)
+	, bExecuteProcessed(0)
+	, bCommandSuccessful(false)
+	, bAutoDelete(true)
+	, Concurrency(EConcurrency::Synchronous)
+{
+	// grab the providers settings here, so we don't access them once the worker thread is launched
+	check(IsInGameThread());
+	FGitSourceControlModule& GitSourceControl = FModuleManager::LoadModuleChecked<FGitSourceControlModule>( "GitSourceControl" );
+	PathToGitBinary = GitSourceControl.AccessSettings().GetBinaryPath();
+	PathToRepositoryRoot = GitSourceControl.GetProvider().GetPathToRepositoryRoot();
+	PathToGameDir = GitSourceControl.GetProvider().GetPathToGameDir();
+}
+
+bool FGitSourceControlCommand::DoWork()
+{
+	bCommandSuccessful = Worker->Execute(*this);
+	FPlatformAtomics::InterlockedExchange(&bExecuteProcessed, 1);
+
+	return bCommandSuccessful;
+}
+
+void FGitSourceControlCommand::Abandon()
+{
+	FPlatformAtomics::InterlockedExchange(&bExecuteProcessed, 1);
+}
+
+void FGitSourceControlCommand::DoThreadedWork()
+{
+	Concurrency = EConcurrency::Asynchronous;
+	DoWork();
+}
diff --git a/Engine/Plugins/Developer/GitSourceControl/Source/GitSourceControl/Private/GitSourceControlCommand.h b/Engine/Plugins/Developer/GitSourceControl/Source/GitSourceControl/Private/GitSourceControlCommand.h
new file mode 100644
--- /dev/null
+++ b/Engine/Plugins/Developer/GitSourceControl/Source/GitSourceControl/Private/GitSourceControlCommand.h
@@ -0,0 +1,76 @@
+// Copyright 1998-2015 Epic Games, Inc. All Rights Reserved.
+
+#pragma once
+
+/**
+ * Used to execute Git commands multi-threaded.
+ */
+class FGitSourceControlCommand : public IQueuedWork
+{
+public:
+
+	FGitSourceControlCommand(const TSharedRef<class ISourceControlOperation, ESPMode::ThreadSafe>& InOperation, const TSharedRef<class IGitSourceControlWorker, ESPMode::ThreadSafe>& InWorker, const FSourceControlOperationComplete& InOperationCompleteDelegate = FSourceControlOperationComplete() );
+
+	/**
+	 * This is where the real thread work is done. All work that is done for
+	 * this queued object should be done from within the call to this function.
+	 */
+	bool DoWork();
+
+	/**
+	 * Tells the queued work that it is being abandoned so that it can do
+	 * per object clean up as needed. This will only be called if it is being
+	 * abandoned before completion. NOTE: This requires the object to delete
+	 * itself using whatever heap it was allocated in.
+	 */
+	virtual void Abandon() override;
+
+	/**
+	 * This method is also used to tell the object to cleanup but not before
+	 * the object has finished it's work.
+	 */ 
+	virtual void DoThreadedWork() override;
+
+public:
+	/** Path to the Git binary */
+	FString PathToGitBinary;
+
+	/** Path to the root of the Git repository: can be the GameDir itself, or any parent directory (found by the "Connect" operation) */
+	FString PathToRepositoryRoot;
+
+	/** Path to the Game directory: shall be inside of the Git repository */
+	FString PathToGameDir;
+
+	/** Name of the current branch (found by the "Connect" operation) */
+	FString BranchName;
+
+	/** Operation we want to perform - contains outward-facing parameters & results */
+	TSharedRef<class ISourceControlOperation, ESPMode::ThreadSafe> Operation;
+
+	/** The object that will actually do the work */
+	TSharedRef<class IGitSourceControlWorker, ESPMode::ThreadSafe> Worker;
+
+	/** Delegate to notify when this operation completes */
+	FSourceControlOperationComplete OperationCompleteDelegate;
+
+	/**If true, this command has been processed by the source control thread*/
+	volatile int32 bExecuteProcessed;
+
+	/**If true, the source control command succeeded*/
+	bool bCommandSuccessful;
+
+	/** If true, this command will be automatically cleaned up in Tick() */
+	bool bAutoDelete;
+
+	/** Whether we are running multi-treaded or not*/
+	EConcurrency::Type Concurrency;
+
+	/** Files to perform this operation on */
+	TArray< FString > Files;
+
+	/**Info and/or warning message message storage*/
+	TArray< FString > InfoMessages;
+
+	/**Potential error message storage*/
+	TArray< FString > ErrorMessages;
+};
\ No newline at end of file
diff --git a/Engine/Plugins/Developer/GitSourceControl/Source/GitSourceControl/Private/GitSourceControlModule.cpp b/Engine/Plugins/Developer/GitSourceControl/Source/GitSourceControl/Private/GitSourceControlModule.cpp
new file mode 100644
--- /dev/null
+++ b/Engine/Plugins/Developer/GitSourceControl/Source/GitSourceControl/Private/GitSourceControlModule.cpp
@@ -0,0 +1,67 @@
+// Copyright 1998-2015 Epic Games, Inc. All Rights Reserved.
+
+#include "GitSourceControlPrivatePCH.h"
+#include "GitSourceControlModule.h"
+#include "ModuleManager.h"
+#include "ISourceControlModule.h"
+#include "GitSourceControlSettings.h"
+#include "GitSourceControlOperations.h"
+#include "Runtime/Core/Public/Features/IModularFeatures.h"
+
+#define LOCTEXT_NAMESPACE "GitSourceControl"
+
+template<typename Type>
+static TSharedRef<IGitSourceControlWorker, ESPMode::ThreadSafe> CreateWorker()
+{
+	return MakeShareable( new Type() );
+}
+
+void FGitSourceControlModule::StartupModule()
+{
+	// Register our operations
+	GitSourceControlProvider.RegisterWorker( "Connect", FGetGitSourceControlWorker::CreateStatic( &CreateWorker<FGitConnectWorker> ) );
+	// Note: this provider does not uses the "CheckOut" command, which is a Perforce "lock", as Git has no lock command (all tracked files in the working copy are always already checked-out).
+	GitSourceControlProvider.RegisterWorker( "UpdateStatus", FGetGitSourceControlWorker::CreateStatic( &CreateWorker<FGitUpdateStatusWorker> ) );
+	GitSourceControlProvider.RegisterWorker( "MarkForAdd", FGetGitSourceControlWorker::CreateStatic( &CreateWorker<FGitMarkForAddWorker> ) );
+	GitSourceControlProvider.RegisterWorker( "Delete", FGetGitSourceControlWorker::CreateStatic( &CreateWorker<FGitDeleteWorker> ) );
+	GitSourceControlProvider.RegisterWorker( "Revert", FGetGitSourceControlWorker::CreateStatic( &CreateWorker<FGitRevertWorker> ) );
+	// @todo: git fetch remote(s) to be able to show files not up-to-date with the serveur
+//	GitSourceControlProvider.RegisterWorker( "Sync", FGetGitSourceControlWorker::CreateStatic( &CreateWorker<FGitSyncWorker> ) );
+	GitSourceControlProvider.RegisterWorker( "CheckIn", FGetGitSourceControlWorker::CreateStatic( &CreateWorker<FGitCheckInWorker> ) );
+	GitSourceControlProvider.RegisterWorker( "Copy", FGetGitSourceControlWorker::CreateStatic( &CreateWorker<FGitCopyWorker> ) );
+	GitSourceControlProvider.RegisterWorker( "Resolve", FGetGitSourceControlWorker::CreateStatic( &CreateWorker<FGitResolveWorker> ) );
+
+	// load our settings
+	GitSourceControlSettings.LoadSettings();
+
+	// Bind our source control provider to the editor
+	IModularFeatures::Get().RegisterModularFeature( "SourceControl", &GitSourceControlProvider );
+}
+
+void FGitSourceControlModule::ShutdownModule()
+{
+	// shut down the provider, as this module is going away
+	GitSourceControlProvider.Close();
+
+	// unbind provider from editor
+	IModularFeatures::Get().UnregisterModularFeature("SourceControl", &GitSourceControlProvider);
+}
+
+FGitSourceControlSettings& FGitSourceControlModule::AccessSettings()
+{
+	return GitSourceControlSettings;
+}
+
+void FGitSourceControlModule::SaveSettings()
+{
+	if (FApp::IsUnattended() || IsRunningCommandlet())
+	{
+		return;
+	}
+
+	GitSourceControlSettings.SaveSettings();
+}
+
+IMPLEMENT_MODULE(FGitSourceControlModule, GitSourceControl);
+
+#undef LOCTEXT_NAMESPACE
diff --git a/Engine/Plugins/Developer/GitSourceControl/Source/GitSourceControl/Private/GitSourceControlModule.h b/Engine/Plugins/Developer/GitSourceControl/Source/GitSourceControl/Private/GitSourceControlModule.h
new file mode 100644
--- /dev/null
+++ b/Engine/Plugins/Developer/GitSourceControl/Source/GitSourceControl/Private/GitSourceControlModule.h
@@ -0,0 +1,34 @@
+// Copyright 1998-2015 Epic Games, Inc. All Rights Reserved.
+
+#pragma once
+
+#include "ISourceControlModule.h"
+#include "GitSourceControlSettings.h"
+#include "GitSourceControlProvider.h"
+
+class FGitSourceControlModule : public IModuleInterface
+{
+public:
+	/** IModuleInterface implementation */
+	virtual void StartupModule() override;
+	virtual void ShutdownModule() override;
+
+	/** Access the Git source control settings */
+	FGitSourceControlSettings& AccessSettings();
+
+	/** Save the Git source control settings */
+	void SaveSettings();
+
+	/** Access the Git source control provider */
+	FGitSourceControlProvider& GetProvider()
+	{
+		return GitSourceControlProvider;
+	}
+
+private:
+	/** The Git source control provider */
+	FGitSourceControlProvider GitSourceControlProvider;
+
+	/** The settings for Git source control */
+	FGitSourceControlSettings GitSourceControlSettings;
+};
\ No newline at end of file
diff --git a/Engine/Plugins/Developer/GitSourceControl/Source/GitSourceControl/Private/GitSourceControlOperations.cpp b/Engine/Plugins/Developer/GitSourceControl/Source/GitSourceControl/Private/GitSourceControlOperations.cpp
new file mode 100644
--- /dev/null
+++ b/Engine/Plugins/Developer/GitSourceControl/Source/GitSourceControl/Private/GitSourceControlOperations.cpp
@@ -0,0 +1,299 @@
+// Copyright 1998-2015 Epic Games, Inc. All Rights Reserved.
+
+#include "GitSourceControlPrivatePCH.h"
+#include "GitSourceControlOperations.h"
+#include "GitSourceControlState.h"
+#include "GitSourceControlCommand.h"
+#include "GitSourceControlModule.h"
+#include "GitSourceControlUtils.h"
+
+#define LOCTEXT_NAMESPACE "GitSourceControl"
+
+FName FGitConnectWorker::GetName() const
+{
+	return "Connect";
+}
+
+bool FGitConnectWorker::Execute(FGitSourceControlCommand& InCommand)
+{
+	check(InCommand.Operation->GetName() == GetName());
+
+	InCommand.bCommandSuccessful = GitSourceControlUtils::FindRootDirectory(InCommand.PathToGameDir, InCommand.PathToRepositoryRoot);
+	if(InCommand.bCommandSuccessful)
+	{
+		InCommand.bCommandSuccessful = GitSourceControlUtils::RunCommand(TEXT("status"), InCommand.PathToGitBinary, InCommand.PathToRepositoryRoot, TArray<FString>(), TArray<FString>(), InCommand.InfoMessages, InCommand.ErrorMessages);
+	}
+	if(!InCommand.bCommandSuccessful || InCommand.ErrorMessages.Num() > 0 || InCommand.InfoMessages.Num() == 0)
+	{
+		// @todo popup to propose to initialize the git repository "git init + .gitignore"
+		StaticCastSharedRef<FConnect>(InCommand.Operation)->SetErrorText(LOCTEXT("NotAWorkingCopyError", "Project is not part of a Git working copy."));
+		// @todo Double error messages (and displayed in reverse order): Perforce distinguish the two errors
+		InCommand.ErrorMessages.Add(LOCTEXT("NotAWorkingCopyErrorHelp", "You should check out a working copy into your project directory.").ToString());
+		InCommand.bCommandSuccessful = false;
+	}
+	else // if(InCommand.bCommandSuccessful)
+	{
+		TArray<FString> Parameters;
+		Parameters.Add("--short HEAD");
+
+		// Get current branche name
+		GitSourceControlUtils::RunCommand(TEXT("symbolic-ref"), InCommand.PathToGitBinary, InCommand.PathToRepositoryRoot, Parameters, TArray<FString>(), InCommand.InfoMessages, InCommand.ErrorMessages);
+		if(InCommand.InfoMessages.Num() == 1)
+		{
+			InCommand.BranchName = InCommand.InfoMessages[0];
+		}
+	}
+
+	return InCommand.bCommandSuccessful;
+}
+
+bool FGitConnectWorker::UpdateStates() const
+{
+	return false;
+}
+
+static FText ParseCommitResults(const TArray<FString>& InResults)
+{
+	if(InResults.Num() >= 1)
+	{
+		const FString& FirstLine = InResults[0];
+		return FText::Format(LOCTEXT("CommitMessage", "Commited {0}."), FText::FromString(FirstLine));
+	}
+	return LOCTEXT("CommitMessageUnknown", "Submitted revision.");
+}
+
+FName FGitCheckInWorker::GetName() const
+{
+	return "CheckIn";
+}
+
+bool FGitCheckInWorker::Execute(FGitSourceControlCommand& InCommand)
+{
+	check(InCommand.Operation->GetName() == GetName());
+
+	TSharedRef<FCheckIn, ESPMode::ThreadSafe> Operation = StaticCastSharedRef<FCheckIn>(InCommand.Operation);
+
+	// make a temp file to place our commit message in
+	FScopedTempFile CommitMsgFile(Operation->GetDescription());
+	if(CommitMsgFile.GetFilename().Len() > 0)
+	{
+		TArray<FString> Parameters;
+		FString ParamCommitMsgFilename = TEXT("--file=\"");
+		ParamCommitMsgFilename += FPaths::ConvertRelativePathToFull(CommitMsgFile.GetFilename());
+		ParamCommitMsgFilename += TEXT("\"");
+		Parameters.Add(ParamCommitMsgFilename);
+
+		InCommand.bCommandSuccessful = GitSourceControlUtils::RunCommit(InCommand.PathToGitBinary, InCommand.PathToRepositoryRoot, Parameters, InCommand.Files, InCommand.InfoMessages, InCommand.ErrorMessages);
+		if(InCommand.bCommandSuccessful)
+		{
+			Operation->SetSuccessMessage(ParseCommitResults(InCommand.InfoMessages));
+			UE_LOG(LogSourceControl, Log, TEXT("FGitCheckInWorker: commit successful"));
+		}
+	}
+
+	// now update the status of our files
+	GitSourceControlUtils::RunUpdateStatus(InCommand.PathToGitBinary, InCommand.PathToRepositoryRoot, InCommand.Files, InCommand.ErrorMessages, States);
+
+	return InCommand.bCommandSuccessful;
+}
+
+bool FGitCheckInWorker::UpdateStates() const
+{
+	return GitSourceControlUtils::UpdateCachedStates(States);
+}
+
+FName FGitMarkForAddWorker::GetName() const
+{
+	return "MarkForAdd";
+}
+
+bool FGitMarkForAddWorker::Execute(FGitSourceControlCommand& InCommand)
+{
+	check(InCommand.Operation->GetName() == GetName());
+
+	InCommand.bCommandSuccessful = GitSourceControlUtils::RunCommand(TEXT("add"), InCommand.PathToGitBinary, InCommand.PathToRepositoryRoot, TArray<FString>(), InCommand.Files, InCommand.InfoMessages, InCommand.ErrorMessages);
+
+	// now update the status of our files
+	GitSourceControlUtils::RunUpdateStatus(InCommand.PathToGitBinary, InCommand.PathToRepositoryRoot, InCommand.Files, InCommand.ErrorMessages, States);
+
+	return InCommand.bCommandSuccessful;
+}
+
+bool FGitMarkForAddWorker::UpdateStates() const
+{
+	return GitSourceControlUtils::UpdateCachedStates(States);
+}
+
+FName FGitDeleteWorker::GetName() const
+{
+	return "Delete";
+}
+
+bool FGitDeleteWorker::Execute(FGitSourceControlCommand& InCommand)
+{
+	check(InCommand.Operation->GetName() == GetName());
+
+	InCommand.bCommandSuccessful = GitSourceControlUtils::RunCommand(TEXT("rm"), InCommand.PathToGitBinary, InCommand.PathToRepositoryRoot, TArray<FString>(), InCommand.Files, InCommand.InfoMessages, InCommand.ErrorMessages);
+
+	// now update the status of our files
+	GitSourceControlUtils::RunUpdateStatus(InCommand.PathToGitBinary, InCommand.PathToRepositoryRoot, InCommand.Files, InCommand.ErrorMessages, States);
+
+	return InCommand.bCommandSuccessful;
+}
+
+bool FGitDeleteWorker::UpdateStates() const
+{
+	return GitSourceControlUtils::UpdateCachedStates(States);
+}
+
+FName FGitRevertWorker::GetName() const
+{
+	return "Revert";
+}
+
+bool FGitRevertWorker::Execute(FGitSourceControlCommand& InCommand)
+{
+	// reset any changes already added in index
+	{
+		InCommand.bCommandSuccessful = GitSourceControlUtils::RunCommand(TEXT("reset"), InCommand.PathToGitBinary, InCommand.PathToRepositoryRoot, TArray<FString>(), InCommand.Files, InCommand.InfoMessages, InCommand.ErrorMessages);
+	}
+
+	// revert any changes in working copy
+	{
+		InCommand.bCommandSuccessful = GitSourceControlUtils::RunCommand(TEXT("checkout"), InCommand.PathToGitBinary, InCommand.PathToRepositoryRoot, TArray<FString>(), InCommand.Files, InCommand.InfoMessages, InCommand.ErrorMessages);
+	}
+
+	// now update the status of our files
+	GitSourceControlUtils::RunUpdateStatus(InCommand.PathToGitBinary, InCommand.PathToRepositoryRoot, InCommand.Files, InCommand.ErrorMessages, States);
+
+	return InCommand.bCommandSuccessful;
+}
+
+bool FGitRevertWorker::UpdateStates() const
+{
+	return GitSourceControlUtils::UpdateCachedStates(States);
+}
+
+FName FGitUpdateStatusWorker::GetName() const
+{
+	return "UpdateStatus";
+}
+
+bool FGitUpdateStatusWorker::Execute(FGitSourceControlCommand& InCommand)
+{
+	check(InCommand.Operation->GetName() == GetName());
+
+	TSharedRef<FUpdateStatus, ESPMode::ThreadSafe> Operation = StaticCastSharedRef<FUpdateStatus>(InCommand.Operation);
+
+	if(InCommand.Files.Num() > 0)
+	{
+		InCommand.bCommandSuccessful = GitSourceControlUtils::RunUpdateStatus(InCommand.PathToGitBinary, InCommand.PathToRepositoryRoot, InCommand.Files, InCommand.ErrorMessages, States);
+
+		if(Operation->ShouldUpdateHistory())
+		{
+			for(const auto& ItFile : InCommand.Files)
+			{
+				TArray<FString> Results;
+				TArray<FString> Parameters;
+
+				Parameters.Add(TEXT("--max-count 100"));
+				Parameters.Add(TEXT("--follow")); // follow file renames
+				Parameters.Add(TEXT("--date=raw"));
+				Parameters.Add(TEXT("--name-status")); // relative filename at this revision, preceded by a status character
+
+				TArray<FString> Files;
+				Files.Add(*ItFile);
+
+				InCommand.bCommandSuccessful &= GitSourceControlUtils::RunCommand(TEXT("log"), InCommand.PathToGitBinary, InCommand.PathToRepositoryRoot, Parameters, Files, Results, InCommand.ErrorMessages);
+				TGitSourceControlHistory History;
+				GitSourceControlUtils::ParseLogResults(Results, History);
+				Histories.Add(*ItFile, History);
+			}
+		}
+	}
+	else
+	{
+		// Perforce "opened files" are those that have been modified (or added/deleted): that is what we get with a simple Git status from the root
+		if(Operation->ShouldGetOpenedOnly())
+		{
+			TArray<FString> Files;
+			Files.Add(FPaths::ConvertRelativePathToFull(FPaths::GameDir()));
+			InCommand.bCommandSuccessful = GitSourceControlUtils::RunUpdateStatus(InCommand.PathToGitBinary, InCommand.PathToRepositoryRoot, Files, InCommand.ErrorMessages, States);
+		}
+	}
+
+	// don't use the ShouldUpdateModifiedState() hint here as it is specific to Perforce: the above normal Git status has already told us this information (like Git and Mercurial)
+
+	return InCommand.bCommandSuccessful;
+}
+
+bool FGitUpdateStatusWorker::UpdateStates() const
+{
+	bool bUpdated = GitSourceControlUtils::UpdateCachedStates(States);
+
+	FGitSourceControlModule& GitSourceControl = FModuleManager::LoadModuleChecked<FGitSourceControlModule>( "GitSourceControl" );
+	FGitSourceControlProvider& Provider = GitSourceControl.GetProvider();
+
+	// add history, if any
+	for(const auto& ItHistory : Histories)
+	{
+		TSharedRef<FGitSourceControlState, ESPMode::ThreadSafe> State = Provider.GetStateInternal(ItHistory.Key);
+		State->History = ItHistory.Value;
+		State->TimeStamp = FDateTime::Now();
+		bUpdated = true;
+	}
+
+	return bUpdated;
+}
+
+FName FGitCopyWorker::GetName() const
+{
+	return "Copy";
+}
+
+bool FGitCopyWorker::Execute(FGitSourceControlCommand& InCommand)
+{
+	check(InCommand.Operation->GetName() == GetName());
+
+	// Copy or Move operation on a single file : Git does not need an explicit copy nor move,
+	// but after a Move the Editor create a redirector file with the old asset name that points to the new asset.
+	// The redirector needs to be commited with the new asset to perform a real rename.
+	// => the following is to "MarkForAdd" the redirector, but it still need to be committed by selecting the whole directory and "check-in"
+	InCommand.bCommandSuccessful = GitSourceControlUtils::RunCommand(TEXT("add"), InCommand.PathToGitBinary, InCommand.PathToRepositoryRoot, TArray<FString>(), InCommand.Files, InCommand.InfoMessages, InCommand.ErrorMessages);
+
+	return InCommand.bCommandSuccessful;
+}
+
+bool FGitCopyWorker::UpdateStates() const
+{
+	return GitSourceControlUtils::UpdateCachedStates(OutStates);
+}
+
+FName FGitResolveWorker::GetName() const
+{
+	return "Resolve";
+}
+
+bool FGitResolveWorker::Execute( class FGitSourceControlCommand& InCommand )
+{
+	check(InCommand.Operation->GetName() == GetName());
+
+	// mark the conflicting files as resolved:
+	{
+		TArray<FString> Results;
+		TArray<FString> Parameters;
+		InCommand.bCommandSuccessful = GitSourceControlUtils::RunCommand(TEXT("add"), InCommand.PathToGitBinary, InCommand.PathToRepositoryRoot, Parameters, InCommand.Files, Results, InCommand.ErrorMessages);
+	}
+
+	// now update the status of our files
+	GitSourceControlUtils::RunUpdateStatus(InCommand.PathToGitBinary, InCommand.PathToRepositoryRoot, InCommand.Files, InCommand.ErrorMessages, States);
+
+	return InCommand.bCommandSuccessful;
+}
+
+bool FGitResolveWorker::UpdateStates() const
+{
+	return GitSourceControlUtils::UpdateCachedStates(States);
+}
+
+#undef LOCTEXT_NAMESPACE
diff --git a/Engine/Plugins/Developer/GitSourceControl/Source/GitSourceControl/Private/GitSourceControlOperations.h b/Engine/Plugins/Developer/GitSourceControl/Source/GitSourceControl/Private/GitSourceControlOperations.h
new file mode 100644
--- /dev/null
+++ b/Engine/Plugins/Developer/GitSourceControl/Source/GitSourceControl/Private/GitSourceControlOperations.h
@@ -0,0 +1,142 @@
+// Copyright 1998-2015 Epic Games, Inc. All Rights Reserved.
+
+#pragma once
+
+#include "IGitSourceControlWorker.h"
+#include "GitSourceControlState.h"
+#include "GitSourceControlRevision.h"
+
+/** Called when first activated on a project, and then at project load time.
+ *  Look for the root directory of the git repository (where the ".git/" subdirectory is located). */
+class FGitConnectWorker : public IGitSourceControlWorker
+{
+public:
+	virtual ~FGitConnectWorker() {}
+	// IGitSourceControlWorker interface
+	virtual FName GetName() const override;
+	virtual bool Execute(class FGitSourceControlCommand& InCommand) override;
+	virtual bool UpdateStates() const override;
+};
+
+/** Commit (check-in) a set of file to the local depot. */
+class FGitCheckInWorker : public IGitSourceControlWorker
+{
+public:
+	virtual ~FGitCheckInWorker() {}
+	// IGitSourceControlWorker interface
+	virtual FName GetName() const override;
+	virtual bool Execute(class FGitSourceControlCommand& InCommand) override;
+	virtual bool UpdateStates() const override;
+
+public:
+	/** Temporary states for results */
+	TArray<FGitSourceControlState> States;
+};
+
+/** Add an untraked file to source control (so only a subset of the git add command). */
+class FGitMarkForAddWorker : public IGitSourceControlWorker
+{
+public:
+	virtual ~FGitMarkForAddWorker() {}
+	// IGitSourceControlWorker interface
+	virtual FName GetName() const override;
+	virtual bool Execute(class FGitSourceControlCommand& InCommand) override;
+	virtual bool UpdateStates() const override;
+
+public:
+	/** Temporary states for results */
+	TArray<FGitSourceControlState> States;
+};
+
+/** Delete a file and remove it from source control. */
+class FGitDeleteWorker : public IGitSourceControlWorker
+{
+public:
+	virtual ~FGitDeleteWorker() {}
+	// IGitSourceControlWorker interface
+	virtual FName GetName() const override;
+	virtual bool Execute(class FGitSourceControlCommand& InCommand) override;
+	virtual bool UpdateStates() const override;
+
+public:
+	/** Map of filenames to Git state */
+	TArray<FGitSourceControlState> States;
+};
+
+/** Revert any change to a file to its state on the local depot. */
+class FGitRevertWorker : public IGitSourceControlWorker
+{
+public:
+	virtual ~FGitRevertWorker() {}
+	// IGitSourceControlWorker interface
+	virtual FName GetName() const override;
+	virtual bool Execute(class FGitSourceControlCommand& InCommand) override;
+	virtual bool UpdateStates() const override;
+
+public:
+	/** Map of filenames to Git state */
+	TArray<FGitSourceControlState> States;
+};
+
+/** @todo: git fetch remote(s) to be able to show files not up-to-date with the serveur
+class FGitSyncWorker : public IGitSourceControlWorker
+{
+public:
+	virtual ~FGitSyncWorker() {}
+	// IGitSourceControlWorker interface
+	virtual FName GetName() const override;
+	virtual bool Execute(class FGitSourceControlCommand& InCommand) override;
+	virtual bool UpdateStates() const override;
+
+public:
+	/// Map of filenames to Git state
+	TArray<FGitSourceControlState> States;
+};
+*/
+
+/** Get source control status of files on local working copy. */
+class FGitUpdateStatusWorker : public IGitSourceControlWorker
+{
+public:
+	virtual ~FGitUpdateStatusWorker() {}
+	// IGitSourceControlWorker interface
+	virtual FName GetName() const override;
+	virtual bool Execute(class FGitSourceControlCommand& InCommand) override;
+	virtual bool UpdateStates() const override;
+
+public:
+	/** Temporary states for results */
+	TArray<FGitSourceControlState> States;
+
+	/** Map of filenames to history */
+	TMap<FString, TGitSourceControlHistory> Histories;
+};
+
+/** Copy or Move operation on a single file */
+class FGitCopyWorker : public IGitSourceControlWorker
+{
+public:
+	virtual ~FGitCopyWorker() {}
+	// IGitSourceControlWorker interface
+	virtual FName GetName() const override;
+	virtual bool Execute(class FGitSourceControlCommand& InCommand) override;
+	virtual bool UpdateStates() const override;
+
+public:
+	/** Temporary states for results */
+	TArray<FGitSourceControlState> OutStates;
+};
+
+/** git add to mark a conflict as resolved */
+class FGitResolveWorker : public IGitSourceControlWorker
+{
+public:
+	virtual ~FGitResolveWorker() {}
+	virtual FName GetName() const override;
+	virtual bool Execute(class FGitSourceControlCommand& InCommand) override;
+	virtual bool UpdateStates() const override;
+	
+private:
+	/** Temporary states for results */
+	TArray<FGitSourceControlState> States;
+};
diff --git a/Engine/Plugins/Developer/GitSourceControl/Source/GitSourceControl/Private/GitSourceControlPrivatePCH.h b/Engine/Plugins/Developer/GitSourceControl/Source/GitSourceControl/Private/GitSourceControlPrivatePCH.h
new file mode 100644
--- /dev/null
+++ b/Engine/Plugins/Developer/GitSourceControl/Source/GitSourceControl/Private/GitSourceControlPrivatePCH.h
@@ -0,0 +1,8 @@
+// Copyright 1998-2015 Epic Games, Inc. All Rights Reserved.
+
+#pragma once
+
+#include "Core.h"
+#include "SlateBasics.h"
+#include "EditorStyle.h"
+#include "ISourceControlModule.h"
diff --git a/Engine/Plugins/Developer/GitSourceControl/Source/GitSourceControl/Private/GitSourceControlProvider.cpp b/Engine/Plugins/Developer/GitSourceControl/Source/GitSourceControl/Private/GitSourceControlProvider.cpp
new file mode 100644
--- /dev/null
+++ b/Engine/Plugins/Developer/GitSourceControl/Source/GitSourceControl/Private/GitSourceControlProvider.cpp
@@ -0,0 +1,330 @@
+// Copyright 1998-2015 Epic Games, Inc. All Rights Reserved.
+
+#include "GitSourceControlPrivatePCH.h"
+#include "GitSourceControlProvider.h"
+#include "GitSourceControlCommand.h"
+#include "ISourceControlModule.h"
+#include "GitSourceControlModule.h"
+#include "GitSourceControlSettings.h"
+#include "GitSourceControlOperations.h"
+#include "GitSourceControlUtils.h"
+#include "SGitSourceControlSettings.h"
+#include "MessageLog.h"
+#include "ScopedSourceControlProgress.h"
+
+#define LOCTEXT_NAMESPACE "GitSourceControl"
+
+static FName ProviderName("Git");
+
+void FGitSourceControlProvider::Init(bool bForceConnection)
+{
+	PathToGameDir = FPaths::ConvertRelativePathToFull(FPaths::GameDir());
+
+	CheckGitAvailability();
+
+	if(bForceConnection)
+	{
+		// First connection, to find the path to the root git directory (if any)
+		GitSourceControlUtils::FindRootDirectory(PathToGameDir, PathToRepositoryRoot);
+	}
+}
+
+void FGitSourceControlProvider::CheckGitAvailability()
+{
+	FGitSourceControlModule& GitSourceControl = FModuleManager::LoadModuleChecked<FGitSourceControlModule>("GitSourceControl");
+	const FString& PathToGitBinary = GitSourceControl.AccessSettings().GetBinaryPath();
+	bGitAvailable = GitSourceControlUtils::CheckGitAvailability(PathToGitBinary);
+}
+
+void FGitSourceControlProvider::Close()
+{
+	// clear the cache
+	StateCache.Empty();
+}
+
+TSharedRef<FGitSourceControlState, ESPMode::ThreadSafe> FGitSourceControlProvider::GetStateInternal(const FString& Filename)
+{
+	TSharedRef<FGitSourceControlState, ESPMode::ThreadSafe>* State = StateCache.Find(Filename);
+	if(State != NULL)
+	{
+		// found cached item
+		return (*State);
+	}
+	else
+	{
+		// cache an unknown state for this item
+		TSharedRef<FGitSourceControlState, ESPMode::ThreadSafe> NewState = MakeShareable( new FGitSourceControlState(Filename) );
+		StateCache.Add(Filename, NewState);
+		return NewState;
+	}
+}
+
+FText FGitSourceControlProvider::GetStatusText() const
+{
+	FFormatNamedArguments Args;
+	Args.Add( TEXT("IsEnabled"), IsEnabled() ? LOCTEXT("Yes", "Yes") : LOCTEXT("No", "No") );
+	Args.Add( TEXT("RepositoryName"), FText::FromString(PathToRepositoryRoot) );
+	Args.Add( TEXT("BranchName"), FText::FromString(BranchName) );
+
+	return FText::Format( NSLOCTEXT("Status", "Provider: Git\nEnabledLabel", "Enabled: {IsEnabled}\nRepository: {RepositoryName}\nBranch: {BranchName}"), Args );
+}
+
+bool FGitSourceControlProvider::IsEnabled() const
+{
+	return bGitAvailable;
+}
+
+bool FGitSourceControlProvider::IsAvailable() const
+{
+	return bGitAvailable;
+}
+
+const FName& FGitSourceControlProvider::GetName(void) const
+{
+	return ProviderName;
+}
+
+ECommandResult::Type FGitSourceControlProvider::GetState( const TArray<FString>& InFiles, TArray< TSharedRef<ISourceControlState, ESPMode::ThreadSafe> >& OutState, EStateCacheUsage::Type InStateCacheUsage )
+{
+	if(!IsEnabled())
+	{
+		return ECommandResult::Failed;
+	}
+
+	TArray<FString> AbsoluteFiles = SourceControlHelpers::AbsoluteFilenames(InFiles);
+
+	if(InStateCacheUsage == EStateCacheUsage::ForceUpdate)
+	{
+		Execute(ISourceControlOperation::Create<FUpdateStatus>(), AbsoluteFiles);
+	}
+
+	for(const auto& AbsoluteFile : AbsoluteFiles)
+	{
+		OutState.Add(GetStateInternal(*AbsoluteFile));
+	}
+
+	return ECommandResult::Succeeded;
+}
+
+void FGitSourceControlProvider::RegisterSourceControlStateChanged( const FSourceControlStateChanged::FDelegate& SourceControlStateChanged )
+{
+	OnSourceControlStateChanged.Add( SourceControlStateChanged );
+}
+
+void FGitSourceControlProvider::UnregisterSourceControlStateChanged( const FSourceControlStateChanged::FDelegate& SourceControlStateChanged )
+{
+	OnSourceControlStateChanged.DEPRECATED_Remove( SourceControlStateChanged );
+}
+
+FDelegateHandle FGitSourceControlProvider::RegisterSourceControlStateChanged_Handle( const FSourceControlStateChanged::FDelegate& SourceControlStateChanged )
+{
+	return OnSourceControlStateChanged.Add( SourceControlStateChanged );
+}
+
+void FGitSourceControlProvider::UnregisterSourceControlStateChanged_Handle( FDelegateHandle Handle )
+{
+	OnSourceControlStateChanged.Remove( Handle );
+}
+
+ECommandResult::Type FGitSourceControlProvider::Execute( const TSharedRef<ISourceControlOperation, ESPMode::ThreadSafe>& InOperation, const TArray<FString>& InFiles, EConcurrency::Type InConcurrency, const FSourceControlOperationComplete& InOperationCompleteDelegate )
+{
+	if(!IsEnabled())
+	{
+		return ECommandResult::Failed;
+	}
+
+	TArray<FString> AbsoluteFiles = SourceControlHelpers::AbsoluteFilenames(InFiles);
+
+	// Query to see if we allow this operation
+	TSharedPtr<IGitSourceControlWorker, ESPMode::ThreadSafe> Worker = CreateWorker(InOperation->GetName());
+	if(!Worker.IsValid())
+	{
+		// this operation is unsupported by this source control provider
+		FFormatNamedArguments Arguments;
+		Arguments.Add( TEXT("OperationName"), FText::FromName(InOperation->GetName()) );
+		Arguments.Add( TEXT("ProviderName"), FText::FromName(GetName()) );
+		FMessageLog("SourceControl").Error(FText::Format(LOCTEXT("UnsupportedOperation", "Operation '{OperationName}' not supported by source control provider '{ProviderName}'"), Arguments));
+		return ECommandResult::Failed;
+	}
+
+	FGitSourceControlCommand* Command = new FGitSourceControlCommand(InOperation, Worker.ToSharedRef());
+	Command->Files = AbsoluteFiles;
+	Command->OperationCompleteDelegate = InOperationCompleteDelegate;
+
+	// fire off operation
+	if(InConcurrency == EConcurrency::Synchronous)
+	{
+		Command->bAutoDelete = false;
+// @todo: temporary debug logs
+//		UE_LOG(LogSourceControl, Log, TEXT("ExecuteSynchronousCommand: '%s'"), *InOperation->GetName().ToString());
+		return ExecuteSynchronousCommand(*Command, InOperation->GetInProgressString());
+	}
+	else
+	{
+		Command->bAutoDelete = true;
+// @todo: temporary debug logs
+//		UE_LOG(LogSourceControl, Log, TEXT("IssueCommand(Asynchronous): '%s'"), *InOperation->GetName().ToString());
+		return IssueCommand(*Command);
+	}
+}
+
+bool FGitSourceControlProvider::CanCancelOperation( const TSharedRef<ISourceControlOperation, ESPMode::ThreadSafe>& InOperation ) const
+{
+	return false;
+}
+
+void FGitSourceControlProvider::CancelOperation( const TSharedRef<ISourceControlOperation, ESPMode::ThreadSafe>& InOperation )
+{
+}
+
+bool FGitSourceControlProvider::UsesLocalReadOnlyState() const
+{
+	return false;
+}
+
+TSharedPtr<IGitSourceControlWorker, ESPMode::ThreadSafe> FGitSourceControlProvider::CreateWorker(const FName& InOperationName) const
+{
+	const FGetGitSourceControlWorker* Operation = WorkersMap.Find(InOperationName);
+	if(Operation != NULL)
+	{
+		return Operation->Execute();
+	}
+		
+	return NULL;
+}
+
+void FGitSourceControlProvider::RegisterWorker( const FName& InName, const FGetGitSourceControlWorker& InDelegate )
+{
+	WorkersMap.Add( InName, InDelegate );
+}
+
+void FGitSourceControlProvider::OutputCommandMessages(const FGitSourceControlCommand& InCommand) const
+{
+	FMessageLog SourceControlLog("SourceControl");
+
+	for(int32 ErrorIndex = 0; ErrorIndex < InCommand.ErrorMessages.Num(); ++ErrorIndex)
+	{
+		SourceControlLog.Error(FText::FromString(InCommand.ErrorMessages[ErrorIndex]));
+	}
+
+	for(int32 InfoIndex = 0; InfoIndex < InCommand.InfoMessages.Num(); ++InfoIndex)
+	{
+		SourceControlLog.Info(FText::FromString(InCommand.InfoMessages[InfoIndex]));
+	}
+}
+
+void FGitSourceControlProvider::Tick()
+{	
+	bool bStatesUpdated = false;
+	for(int32 CommandIndex = 0; CommandIndex < CommandQueue.Num(); ++CommandIndex)
+	{
+		FGitSourceControlCommand& Command = *CommandQueue[CommandIndex];
+		if (Command.bExecuteProcessed)
+		{
+			// Remove command from the queue
+			CommandQueue.RemoveAt(CommandIndex);
+
+			// let command update the states of any files
+			bStatesUpdated |= Command.Worker->UpdateStates();
+
+			// if connect operation: obtain the path to the Git root directory of the repository
+			if(Command.Operation->GetName() == "Connect" && Command.bCommandSuccessful)
+			{
+				PathToRepositoryRoot = Command.PathToRepositoryRoot;
+				BranchName = Command.BranchName;
+			}
+
+			// dump any messages to output log
+			OutputCommandMessages(Command);
+
+			// run the completion delegate callback if we have one bound
+			ECommandResult::Type Result = Command.bCommandSuccessful ? ECommandResult::Succeeded : ECommandResult::Failed;
+			Command.OperationCompleteDelegate.ExecuteIfBound(Command.Operation, Result);
+
+			// commands that are left in the array during a tick need to be deleted
+			if(Command.bAutoDelete)
+			{
+				// Only delete commands that are not running 'synchronously'
+				delete &Command;
+			}
+			
+			// only do one command per tick loop, as we dont want concurrent modification 
+			// of the command queue (which can happen in the completion delegate)
+			break;
+		}
+	}
+
+	if(bStatesUpdated)
+	{
+		OnSourceControlStateChanged.Broadcast();
+	}
+}
+
+TArray< TSharedRef<ISourceControlLabel> > FGitSourceControlProvider::GetLabels( const FString& InMatchingSpec ) const
+{
+	TArray< TSharedRef<ISourceControlLabel> > Tags;
+
+	return Tags;
+}
+
+#if SOURCE_CONTROL_WITH_SLATE
+TSharedRef<class SWidget> FGitSourceControlProvider::MakeSettingsWidget() const
+{
+	return SNew(SGitSourceControlSettings);
+}
+#endif
+
+ECommandResult::Type FGitSourceControlProvider::ExecuteSynchronousCommand(FGitSourceControlCommand& InCommand, const FText& Task)
+{
+	ECommandResult::Type Result = ECommandResult::Failed;
+
+	// Display the progress dialog if a string was provided
+	{
+		FScopedSourceControlProgress Progress(Task);
+
+		// Issue the command asynchronously...
+		IssueCommand( InCommand );
+
+		// ... then wait for its completion (thus making it synchrounous)
+		while(!InCommand.bExecuteProcessed)
+		{
+			// Tick the command queue and update progress.
+			Tick();
+			
+			Progress.Tick();
+
+			// Sleep for a bit so we don't busy-wait so much.
+			FPlatformProcess::Sleep(0.01f);
+		}
+	
+		// always do one more Tick() to make sure the command queue is cleaned up.
+		Tick();
+
+		if(InCommand.bCommandSuccessful)
+		{
+			Result = ECommandResult::Succeeded;
+		}
+	}
+
+	// Delete the command now (asynchronous commands are deleted in the Tick() method)
+	check(!InCommand.bAutoDelete);
+	delete &InCommand;
+
+	return Result;
+}
+
+ECommandResult::Type FGitSourceControlProvider::IssueCommand(FGitSourceControlCommand& InCommand)
+{
+	if ( GThreadPool != NULL )
+	{
+		// Queue this to our worker thread(s) for resolving
+		GThreadPool->AddQueuedWork(&InCommand);
+		CommandQueue.Add(&InCommand);
+		return ECommandResult::Succeeded;
+	}
+	else
+	{
+		return ECommandResult::Failed;
+	}
+}
+#undef LOCTEXT_NAMESPACE
diff --git a/Engine/Plugins/Developer/GitSourceControl/Source/GitSourceControl/Private/GitSourceControlProvider.h b/Engine/Plugins/Developer/GitSourceControl/Source/GitSourceControl/Private/GitSourceControlProvider.h
new file mode 100644
--- /dev/null
+++ b/Engine/Plugins/Developer/GitSourceControl/Source/GitSourceControl/Private/GitSourceControlProvider.h
@@ -0,0 +1,111 @@
+// Copyright 1998-2015 Epic Games, Inc. All Rights Reserved.
+
+#pragma once
+
+#include "ISourceControlProvider.h"
+#include "IGitSourceControlWorker.h"
+#include "GitSourceControlState.h"
+
+DECLARE_DELEGATE_RetVal(FGitSourceControlWorkerRef, FGetGitSourceControlWorker)
+
+// We disable the deprecation warnings here because otherwise it'll complain about us
+// implementing RegisterSourceControlStateChanged and UnregisterSourceControlStateChanged.  We know
+// that, but we only want it to complain if *others* implement or call these functions.
+//
+// These macros should be removed when those functions are removed.
+PRAGMA_DISABLE_DEPRECATION_WARNINGS
+
+class FGitSourceControlProvider : public ISourceControlProvider
+{
+public:
+	/** Constructor */
+	FGitSourceControlProvider() 
+		: bGitAvailable(false)
+	{
+	}
+
+	/* ISourceControlProvider implementation */
+	virtual void Init(bool bForceConnection = true) override;
+	virtual void Close() override;
+	virtual FText GetStatusText() const override;
+	virtual bool IsEnabled() const override;
+	virtual bool IsAvailable() const override;
+	virtual const FName& GetName(void) const override;
+	virtual ECommandResult::Type GetState( const TArray<FString>& InFiles, TArray< TSharedRef<ISourceControlState, ESPMode::ThreadSafe> >& OutState, EStateCacheUsage::Type InStateCacheUsage ) override;
+	virtual void RegisterSourceControlStateChanged( const FSourceControlStateChanged::FDelegate& SourceControlStateChanged ) override;
+	virtual void UnregisterSourceControlStateChanged( const FSourceControlStateChanged::FDelegate& SourceControlStateChanged ) override;
+	virtual FDelegateHandle RegisterSourceControlStateChanged_Handle(const FSourceControlStateChanged::FDelegate& SourceControlStateChanged) override;
+	virtual void UnregisterSourceControlStateChanged_Handle(FDelegateHandle Handle) override;
+	virtual ECommandResult::Type Execute(const TSharedRef<ISourceControlOperation, ESPMode::ThreadSafe>& InOperation, const TArray<FString>& InFiles, EConcurrency::Type InConcurrency = EConcurrency::Synchronous, const FSourceControlOperationComplete& InOperationCompleteDelegate = FSourceControlOperationComplete()) override;
+	virtual bool CanCancelOperation( const TSharedRef<ISourceControlOperation, ESPMode::ThreadSafe>& InOperation ) const override;
+	virtual void CancelOperation( const TSharedRef<ISourceControlOperation, ESPMode::ThreadSafe>& InOperation ) override;
+	virtual bool UsesLocalReadOnlyState() const override;
+	virtual void Tick() override;
+	virtual TArray< TSharedRef<class ISourceControlLabel> > GetLabels( const FString& InMatchingSpec ) const override;
+#if SOURCE_CONTROL_WITH_SLATE
+	virtual TSharedRef<class SWidget> MakeSettingsWidget() const override;
+#endif
+
+	/**
+	 * Run a Git "version" command to check the availability of the binary.
+	 */
+	void CheckGitAvailability();
+
+	/** Get the path to the root of the Git repository: can be the GameDir itself, or any parent directory */
+	inline const FString& GetPathToRepositoryRoot() const
+	{
+		return PathToRepositoryRoot;
+	}
+
+	/** Get the path to the Game directory: shall be inside of the Git repository */
+	inline const FString& GetPathToGameDir() const
+	{
+		return PathToGameDir;
+	}
+
+	/** Helper function used to update state cache */
+	TSharedRef<FGitSourceControlState, ESPMode::ThreadSafe> GetStateInternal(const FString& Filename);
+
+	/**
+	 * Register a worker with the provider.
+	 * This is used internally so the provider can maintain a map of all available operations.
+	 */
+	void RegisterWorker( const FName& InName, const FGetGitSourceControlWorker& InDelegate );
+
+private:
+
+	/** Is git binary found and working. */
+	bool bGitAvailable;
+
+	/** Helper function for Execute() */
+	TSharedPtr<class IGitSourceControlWorker, ESPMode::ThreadSafe> CreateWorker(const FName& InOperationName) const;
+
+	/** Helper function for running command synchronously. */
+	ECommandResult::Type ExecuteSynchronousCommand(class FGitSourceControlCommand& InCommand, const FText& Task);
+	/** Issue a command asynchronously if possible. */
+	ECommandResult::Type IssueCommand(class FGitSourceControlCommand& InCommand);
+
+	/** Output any messages this command holds */
+	void OutputCommandMessages(const class FGitSourceControlCommand& InCommand) const;
+
+	/** Path to the root of the Git repository: can be the GameDir itself, or any parent directory (found by the "Connect" operation) */
+	FString PathToRepositoryRoot;
+
+	/** Path to the Game directory: shall be inside of the Git repository */
+	FString PathToGameDir;
+
+	/** Name of the current branch (found by the "Connect" operation) */
+	FString BranchName;
+
+	/** State cache */
+	TMap<FString, TSharedRef<class FGitSourceControlState, ESPMode::ThreadSafe> > StateCache;
+
+	/** The currently registered source control operations */
+	TMap<FName, FGetGitSourceControlWorker> WorkersMap;
+
+	/** Queue for commands given by the main thread */
+	TArray < FGitSourceControlCommand* > CommandQueue;
+
+	/** For notifying when the source control states in the cache have changed */
+	FSourceControlStateChanged OnSourceControlStateChanged;
+};
diff --git a/Engine/Plugins/Developer/GitSourceControl/Source/GitSourceControl/Private/GitSourceControlRevision.cpp b/Engine/Plugins/Developer/GitSourceControl/Source/GitSourceControl/Private/GitSourceControlRevision.cpp
new file mode 100644
--- /dev/null
+++ b/Engine/Plugins/Developer/GitSourceControl/Source/GitSourceControl/Private/GitSourceControlRevision.cpp
@@ -0,0 +1,105 @@
+// Copyright 1998-2015 Epic Games, Inc. All Rights Reserved.
+
+#include "GitSourceControlPrivatePCH.h"
+#include "GitSourceControlRevision.h"
+#include "GitSourceControlModule.h"
+#include "GitSourceControlProvider.h"
+#include "GitSourceControlUtils.h"
+#include "SGitSourceControlSettings.h"
+
+#define LOCTEXT_NAMESPACE "GitSourceControl"
+
+bool FGitSourceControlRevision::Get( FString& InOutFilename ) const
+{
+	FGitSourceControlModule& GitSourceControl = FModuleManager::LoadModuleChecked<FGitSourceControlModule>("GitSourceControl");
+	FString PathToGitBinary = GitSourceControl.AccessSettings().GetBinaryPath();
+	FString PathToRepositoryRoot = GitSourceControl.GetProvider().GetPathToRepositoryRoot();
+
+	// if a filename for the temp file wasn't supplied generate a unique-ish one
+	if(InOutFilename.Len() == 0)
+	{
+		// create the diff dir if we don't already have it (SVN wont)
+		IFileManager::Get().MakeDirectory(*FPaths::DiffDir(), true);
+
+		// @todo Bug report: a counter is needed to avoid overlapping files; temp files are not (never?) released by Editor!
+		static int32 TempFileCount = 0;
+		FString TempFileName = FString::Printf(TEXT("%stemp-%d-%s-%s"), *FPaths::DiffDir(), TempFileCount++, *CommitId, *FPaths::GetCleanFilename(Filename));
+		InOutFilename = FPaths::ConvertRelativePathToFull(TempFileName);
+	}
+
+	// Diff against the revision
+	FString Parameter = CommitId;
+	Parameter += TEXT(":");
+	Parameter += *Filename;
+
+	bool bCommandSuccessful = GitSourceControlUtils::RunDumpToFile(PathToGitBinary, PathToRepositoryRoot, Parameter, InOutFilename);
+
+	return bCommandSuccessful;
+}
+
+bool FGitSourceControlRevision::GetAnnotated( TArray<FAnnotationLine>& OutLines ) const
+{
+	return false;
+}
+
+bool FGitSourceControlRevision::GetAnnotated( FString& InOutFilename ) const
+{
+	return false;
+}
+
+const FString& FGitSourceControlRevision::GetFilename() const
+{
+	return Filename;
+}
+
+int32 FGitSourceControlRevision::GetRevisionNumber() const
+{
+	return RevisionNumber;
+}
+
+const FString& FGitSourceControlRevision::GetDescription() const
+{
+	return Description;
+}
+
+const FString& FGitSourceControlRevision::GetUserName() const
+{
+	return UserName;
+}
+
+const FString& FGitSourceControlRevision::GetClientSpec() const
+{
+	static FString EmptyString(TEXT(""));
+	return EmptyString;
+}
+
+const FString& FGitSourceControlRevision::GetAction() const
+{
+	return Action;
+}
+
+TSharedPtr<class ISourceControlRevision, ESPMode::ThreadSafe> FGitSourceControlRevision::GetBranchSource() const
+{
+	// @todo if this revision was copied from some other revision, then that source revision should
+	//       be returned here (this should be determined when history is being fetched)
+	return nullptr;
+}
+
+const FDateTime& FGitSourceControlRevision::GetDate() const
+{
+	return Date;
+}
+
+int32 FGitSourceControlRevision::GetCheckInIdentifier() const
+{
+	// in Git, revisions apply to the whole repository so (in Perforce terms) the revision *is* the changelist
+	return RevisionNumber;
+}
+
+int32 FGitSourceControlRevision::GetFileSize() const
+{
+	// @todo git log does not give us the file size, but we could run a specific command
+	return 0;
+}
+
+#undef LOCTEXT_NAMESPACE
diff --git a/Engine/Plugins/Developer/GitSourceControl/Source/GitSourceControl/Private/GitSourceControlRevision.h b/Engine/Plugins/Developer/GitSourceControl/Source/GitSourceControl/Private/GitSourceControlRevision.h
new file mode 100644
--- /dev/null
+++ b/Engine/Plugins/Developer/GitSourceControl/Source/GitSourceControl/Private/GitSourceControlRevision.h
@@ -0,0 +1,56 @@
+// Copyright 1998-2015 Epic Games, Inc. All Rights Reserved.
+
+#pragma once
+
+#include "ISourceControlRevision.h"
+
+/** Revision of a file, linked to a specific commit */
+class FGitSourceControlRevision : public ISourceControlRevision, public TSharedFromThis<FGitSourceControlRevision, ESPMode::ThreadSafe>
+{
+public:
+	FGitSourceControlRevision()
+		: RevisionNumber(0)
+	{
+	}
+
+	/** ISourceControlRevision interface */
+	virtual bool Get( FString& InOutFilename ) const override;
+	virtual bool GetAnnotated( TArray<FAnnotationLine>& OutLines ) const override;
+	virtual bool GetAnnotated( FString& InOutFilename ) const override;
+	virtual const FString& GetFilename() const override;
+	virtual int32 GetRevisionNumber() const override;
+	virtual const FString& GetDescription() const override;
+	virtual const FString& GetUserName() const override;
+	virtual const FString& GetClientSpec() const override;
+	virtual const FString& GetAction() const override;
+	virtual TSharedPtr<class ISourceControlRevision, ESPMode::ThreadSafe> GetBranchSource() const override;
+	virtual const FDateTime& GetDate() const override;
+	virtual int32 GetCheckInIdentifier() const override;
+	virtual int32 GetFileSize() const override;
+
+public:
+
+	/** The filename this revision refers to */
+	FString Filename;
+
+	/** The SHA1 id of the commit this revision refers to */
+	FString CommitId;
+
+	/** The revision number */
+	int32 RevisionNumber;
+
+	/** The description of this revision */
+	FString Description;
+
+	/** The user that made the change */
+	FString UserName;
+
+	/** The action (add, edit etc.) performed at this revision */
+	FString Action;
+
+	/** The date this revision was made */
+	FDateTime Date;
+};
+
+/** History composed of the last 100 revisions of the file */
+typedef TArray< TSharedRef<FGitSourceControlRevision, ESPMode::ThreadSafe> >	TGitSourceControlHistory;
diff --git a/Engine/Plugins/Developer/GitSourceControl/Source/GitSourceControl/Private/GitSourceControlSettings.cpp b/Engine/Plugins/Developer/GitSourceControl/Source/GitSourceControl/Private/GitSourceControlSettings.cpp
new file mode 100644
--- /dev/null
+++ b/Engine/Plugins/Developer/GitSourceControl/Source/GitSourceControl/Private/GitSourceControlSettings.cpp
@@ -0,0 +1,49 @@
+// Copyright 1998-2015 Epic Games, Inc. All Rights Reserved.
+
+#include "GitSourceControlPrivatePCH.h"
+#include "GitSourceControlSettings.h"
+#include "GitSourceControlModule.h"
+#include "GitSourceControlProvider.h"
+#include "GitSourceControlUtils.h"
+#include "SourceControlHelpers.h"
+
+namespace GitSettingsConstants
+{
+
+/** The section of the ini file we load our settings from */
+static const FString SettingsSection = TEXT("GitSourceControl.GitSourceControlSettings");
+
+}
+
+const FString& FGitSourceControlSettings::GetBinaryPath() const
+{
+	FScopeLock ScopeLock(&CriticalSection);
+	return BinaryPath;
+}
+
+void FGitSourceControlSettings::SetBinaryPath(const FString& InString)
+{
+	FScopeLock ScopeLock(&CriticalSection);
+	BinaryPath = InString;
+}
+
+void FGitSourceControlSettings::LoadSettings()
+{
+	FScopeLock ScopeLock(&CriticalSection);
+	const FString& IniFile = SourceControlHelpers::GetSettingsIni();
+	bool bLoaded = GConfig->GetString(*GitSettingsConstants::SettingsSection, TEXT("BinaryPath"), BinaryPath, IniFile);
+	if(!bLoaded)
+	{
+		BinaryPath = GitSourceControlUtils::FindGitBinaryPath();
+	}
+}
+
+void FGitSourceControlSettings::SaveSettings() const
+{
+	FScopeLock ScopeLock(&CriticalSection);
+	const FString& IniFile = SourceControlHelpers::GetSettingsIni();
+	GConfig->SetString(*GitSettingsConstants::SettingsSection, TEXT("BinaryPath"), *BinaryPath, IniFile);
+
+	FGitSourceControlModule& GitSourceControl = FModuleManager::LoadModuleChecked<FGitSourceControlModule>("GitSourceControl");
+	GitSourceControl.GetProvider().CheckGitAvailability();
+}
\ No newline at end of file
diff --git a/Engine/Plugins/Developer/GitSourceControl/Source/GitSourceControl/Private/GitSourceControlSettings.h b/Engine/Plugins/Developer/GitSourceControl/Source/GitSourceControl/Private/GitSourceControlSettings.h
new file mode 100644
--- /dev/null
+++ b/Engine/Plugins/Developer/GitSourceControl/Source/GitSourceControl/Private/GitSourceControlSettings.h
@@ -0,0 +1,26 @@
+// Copyright 1998-2015 Epic Games, Inc. All Rights Reserved.
+
+#pragma once
+
+class FGitSourceControlSettings
+{
+public:
+	/** Get the Git Binary Path */
+	const FString& GetBinaryPath() const;
+
+	/** Set the Git Binary Path */
+	void SetBinaryPath(const FString& InString);
+
+	/** Load settings from ini file */
+	void LoadSettings();
+
+	/** Save settings to ini file */
+	void SaveSettings() const;
+
+private:
+	/** A critical section for settings access */
+	mutable FCriticalSection CriticalSection;
+
+	/** Git binary path */
+	FString BinaryPath;
+};
\ No newline at end of file
diff --git a/Engine/Plugins/Developer/GitSourceControl/Source/GitSourceControl/Private/GitSourceControlState.cpp b/Engine/Plugins/Developer/GitSourceControl/Source/GitSourceControl/Private/GitSourceControlState.cpp
new file mode 100644
--- /dev/null
+++ b/Engine/Plugins/Developer/GitSourceControl/Source/GitSourceControl/Private/GitSourceControlState.cpp
@@ -0,0 +1,259 @@
+// Copyright 1998-2015 Epic Games, Inc. All Rights Reserved.
+
+#include "GitSourceControlPrivatePCH.h"
+#include "GitSourceControlState.h"
+
+#define LOCTEXT_NAMESPACE "GitSourceControl.State"
+
+int32 FGitSourceControlState::GetHistorySize() const
+{
+	return History.Num();
+}
+
+TSharedPtr<class ISourceControlRevision, ESPMode::ThreadSafe> FGitSourceControlState::GetHistoryItem( int32 HistoryIndex ) const
+{
+	check(History.IsValidIndex(HistoryIndex));
+	return History[HistoryIndex];
+}
+
+TSharedPtr<class ISourceControlRevision, ESPMode::ThreadSafe> FGitSourceControlState::FindHistoryRevision( int32 RevisionNumber ) const
+{
+	for(const auto& Revision : History)
+	{
+		if(Revision->GetRevisionNumber() == RevisionNumber)
+		{
+			return Revision;
+		}
+	}
+
+	return NULL;
+}
+
+TSharedPtr<class ISourceControlRevision, ESPMode::ThreadSafe> FGitSourceControlState::GetBaseRevForMerge() const
+{
+	// @todo get revision of the merge-base (https://www.kernel.org/pub/software/scm/git/docs/git-merge-base.html)
+	return nullptr;
+}
+
+// @todo add Slate icons for git specific states (Added vs Modified, Copied vs Conflicted...)
+FName FGitSourceControlState::GetIconName() const
+{
+	switch(WorkingCopyState)
+	{
+	case EWorkingCopyState::Modified:
+		return FName("Subversion.CheckedOut");
+	case EWorkingCopyState::Added:
+	case EWorkingCopyState::Renamed:
+	case EWorkingCopyState::Copied:
+	case EWorkingCopyState::Deleted: // @todo New dedicated icon for removed/deleted (MarkedForDelete) files : Deleted files does not always show in Editor (but need to be checked-in) : only when deleted externaly !
+		return FName("Subversion.OpenForAdd");
+	case EWorkingCopyState::Conflicted:
+		return FName("Subversion.NotAtHeadRevision");
+	case EWorkingCopyState::NotControlled:
+		return FName("Subversion.NotInDepot");
+		UE_LOG(LogSourceControl, Log, TEXT("EWorkingCopyState::Deleted"));
+
+	case EWorkingCopyState::Missing: // @todo Missing files does not currently show in Editor (but should probably)
+		UE_LOG(LogSourceControl, Log, TEXT("EWorkingCopyState::Missing"));
+//	case EWorkingCopyState::Unchanged:
+		// Unchanged is the same as "Pristine" (not checked out) for Perforce, ie no icon
+	}
+
+	return NAME_None;
+}
+
+FName FGitSourceControlState::GetSmallIconName() const
+{
+	switch(WorkingCopyState)
+	{
+	case EWorkingCopyState::Unchanged:
+		return FName("Subversion.CheckedOut_Small");
+	case EWorkingCopyState::Added:
+	case EWorkingCopyState::Renamed:
+	case EWorkingCopyState::Copied:
+	case EWorkingCopyState::Deleted: // @todo New dedicated icon for removed/deleted (MarkedForDelete) files : Deleted files does not always show in Editor (but need to be checked-in) : only when deleted externaly !
+		return FName("Subversion.OpenForAdd_Small");
+	case EWorkingCopyState::Conflicted:
+		return FName("Subversion.NotAtHeadRevision_Small");
+	case EWorkingCopyState::NotControlled:
+		return FName("Subversion.NotInDepot_Small");
+
+	case EWorkingCopyState::Missing: // @todo Missing files does not currently show in Editor (but should probably)
+		UE_LOG(LogSourceControl, Log, TEXT("EWorkingCopyState::Missing"));
+//	case EWorkingCopyState::Unchanged:
+		// Unchanged is the same as "Pristine" (not checked out) for Perforce, ie no icon
+	}
+
+	return NAME_None;
+}
+
+FText FGitSourceControlState::GetDisplayName() const
+{
+	switch(WorkingCopyState)
+	{
+	case EWorkingCopyState::Unknown:
+		return LOCTEXT("Unknown", "Unknown");
+	case EWorkingCopyState::Unchanged:
+		return LOCTEXT("Unchanged", "Unchanged");
+	case EWorkingCopyState::Added:
+		return LOCTEXT("Added", "Added");
+	case EWorkingCopyState::Deleted:
+		return LOCTEXT("Deleted", "Deleted");
+	case EWorkingCopyState::Modified:
+		return LOCTEXT("Modified", "Modified");
+	case EWorkingCopyState::Renamed:
+		return LOCTEXT("Renamed", "Renamed");
+	case EWorkingCopyState::Copied:
+		return LOCTEXT("Copied", "Copied");
+	case EWorkingCopyState::Conflicted:
+		return LOCTEXT("ContentsConflict", "Contents Conflict");
+	case EWorkingCopyState::Ignored:
+		return LOCTEXT("Ignored", "Ignored");
+	case EWorkingCopyState::Merged:
+		return LOCTEXT("Merged", "Merged");
+	case EWorkingCopyState::NotControlled:
+		return LOCTEXT("NotControlled", "Not Under Source Control");
+	case EWorkingCopyState::Missing:
+		return LOCTEXT("Missing", "Missing");
+	}
+
+	return FText();
+}
+
+FText FGitSourceControlState::GetDisplayTooltip() const
+{
+	switch(WorkingCopyState)
+	{
+	case EWorkingCopyState::Unknown:
+		return LOCTEXT("Unknown_Tooltip", "Unknown source control state");
+	case EWorkingCopyState::Unchanged:
+		return LOCTEXT("Pristine_Tooltip", "There are no modifications");
+	case EWorkingCopyState::Added:
+		return LOCTEXT("Added_Tooltip", "Item is scheduled for addition");
+	case EWorkingCopyState::Deleted:
+		return LOCTEXT("Deleted_Tooltip", "Item is scheduled for deletion");
+	case EWorkingCopyState::Modified:
+		return LOCTEXT("Modified_Tooltip", "Item has been modified");
+	case EWorkingCopyState::Conflicted:
+		return LOCTEXT("ContentsConflict_Tooltip", "The contents (as opposed to the properties) of the item conflict with updates received from the repository.");
+	case EWorkingCopyState::Ignored:
+		return LOCTEXT("Ignored_Tooltip", "Item is being ignored.");
+	case EWorkingCopyState::Merged:
+		return LOCTEXT("Merged_Tooltip", "Item has been merged.");
+	case EWorkingCopyState::NotControlled:
+		return LOCTEXT("NotControlled_Tooltip", "Item is not under version control.");
+	case EWorkingCopyState::Missing:
+		return LOCTEXT("Missing_Tooltip", "Item is missing (e.g., you moved or deleted it without using Git). This also indicates that a directory is incomplete (a checkout or update was interrupted).");
+	}
+
+	return FText();
+}
+
+const FString& FGitSourceControlState::GetFilename() const
+{
+	return LocalFilename;
+}
+
+const FDateTime& FGitSourceControlState::GetTimeStamp() const
+{
+	return TimeStamp;
+}
+
+// TODO : missing ?
+// @todo Test: don't show deleted as they should not appear? 
+//	case EWorkingCopyState::Deleted:
+//	case EWorkingCopyState::Missing:
+// Deleted and Missing assets cannot appear in the Content Browser
+bool FGitSourceControlState::CanCheckIn() const
+{
+	return WorkingCopyState == EWorkingCopyState::Added
+		|| WorkingCopyState == EWorkingCopyState::Deleted
+		|| WorkingCopyState == EWorkingCopyState::Modified
+		|| WorkingCopyState == EWorkingCopyState::Renamed;
+}
+
+bool FGitSourceControlState::CanCheckout() const
+{
+	return false; // With Git all tracked files in the working copy are always already checked-out (as opposed to Perforce)
+}
+
+bool FGitSourceControlState::IsCheckedOut() const
+{
+	return IsSourceControlled(); // With Git all tracked files in the working copy are always checked-out (as opposed to Perforce)
+}
+
+bool FGitSourceControlState::IsCheckedOutOther(FString* Who) const
+{
+	return false; // Git does not lock checked-out files as Perforce does
+}
+
+bool FGitSourceControlState::IsCurrent() const
+{
+	return true; // @todo check the state of the HEAD versus the state of tracked branch on remote
+}
+
+bool FGitSourceControlState::IsSourceControlled() const
+{
+	return WorkingCopyState != EWorkingCopyState::NotControlled && WorkingCopyState != EWorkingCopyState::Ignored && WorkingCopyState != EWorkingCopyState::Unknown;
+}
+
+bool FGitSourceControlState::IsAdded() const
+{
+	return WorkingCopyState == EWorkingCopyState::Added;
+}
+
+bool FGitSourceControlState::IsDeleted() const
+{
+	return WorkingCopyState == EWorkingCopyState::Deleted;
+}
+
+bool FGitSourceControlState::IsIgnored() const
+{
+	return WorkingCopyState == EWorkingCopyState::Ignored;
+}
+
+bool FGitSourceControlState::CanEdit() const
+{
+	return true; // With Git all files in the working copy are always editable (as opposed to Perforce)
+}
+
+bool FGitSourceControlState::IsUnknown() const
+{
+	return WorkingCopyState == EWorkingCopyState::Unknown;
+}
+
+bool FGitSourceControlState::IsModified() const
+{
+	// Warning: for Perforce, a checked-out file is locked for modification (whereas with Git all tracked files are checked-out),
+	// so for a clean "check-in" (commit) checked-out files unmodified should be removed from the changeset (the index)
+	// http://stackoverflow.com/questions/12357971/what-does-revert-unchanged-files-mean-in-perforce
+	//
+	// Thus, before check-in UE4 Editor call RevertUnchangedFiles() in PromptForCheckin() and CheckinFiles().
+	//
+	// So here we must take care to enumerate all states that need to be commited,
+	// all other will be discarded :
+	//  - Unknown
+	//  - Unchanged
+	//  - NotControlled
+	//  - Ignored
+	return WorkingCopyState == EWorkingCopyState::Added
+		|| WorkingCopyState == EWorkingCopyState::Deleted
+		|| WorkingCopyState == EWorkingCopyState::Modified
+		|| WorkingCopyState == EWorkingCopyState::Renamed
+		|| WorkingCopyState == EWorkingCopyState::Copied
+		|| WorkingCopyState == EWorkingCopyState::Conflicted
+		|| WorkingCopyState == EWorkingCopyState::Missing;
+}
+
+
+bool FGitSourceControlState::CanAdd() const
+{
+	return WorkingCopyState == EWorkingCopyState::NotControlled;
+}
+
+bool FGitSourceControlState::IsConflicted() const
+{
+	return WorkingCopyState == EWorkingCopyState::Conflicted;
+}
+
+#undef LOCTEXT_NAMESPACE
diff --git a/Engine/Plugins/Developer/GitSourceControl/Source/GitSourceControl/Private/GitSourceControlState.h b/Engine/Plugins/Developer/GitSourceControl/Source/GitSourceControl/Private/GitSourceControlState.h
new file mode 100644
--- /dev/null
+++ b/Engine/Plugins/Developer/GitSourceControl/Source/GitSourceControl/Private/GitSourceControlState.h
@@ -0,0 +1,75 @@
+// Copyright 1998-2015 Epic Games, Inc. All Rights Reserved.
+
+#pragma once
+
+#include "ISourceControlState.h"
+#include "GitSourceControlRevision.h"
+
+namespace EWorkingCopyState
+{
+	enum Type
+	{
+		Unknown,
+		Unchanged, // called "clean" in SVN, "Pristine" in Perforce
+		Added,
+		Deleted,
+		Modified,
+		Renamed,
+		Copied,
+		Missing,
+		Conflicted,
+		Merged,
+		NotControlled,
+		Ignored,
+	};
+}
+
+class FGitSourceControlState : public ISourceControlState, public TSharedFromThis<FGitSourceControlState, ESPMode::ThreadSafe>
+{
+public:
+	FGitSourceControlState( const FString& InLocalFilename )
+		: LocalFilename(InLocalFilename)
+		, WorkingCopyState(EWorkingCopyState::Unknown)
+		, TimeStamp(0)
+	{
+	}
+
+	/** ISourceControlState interface */
+	virtual int32 GetHistorySize() const override;
+	virtual TSharedPtr<class ISourceControlRevision, ESPMode::ThreadSafe> GetHistoryItem(int32 HistoryIndex) const override;
+	virtual TSharedPtr<class ISourceControlRevision, ESPMode::ThreadSafe> FindHistoryRevision(int32 RevisionNumber) const override;
+	virtual TSharedPtr<class ISourceControlRevision, ESPMode::ThreadSafe> GetBaseRevForMerge() const override;
+	virtual FName GetIconName() const override;
+	virtual FName GetSmallIconName() const override;
+	virtual FText GetDisplayName() const override;
+	virtual FText GetDisplayTooltip() const override;
+	virtual const FString& GetFilename() const override;
+	virtual const FDateTime& GetTimeStamp() const override;
+	virtual bool CanCheckIn() const override;
+	virtual bool CanCheckout() const override;
+	virtual bool IsCheckedOut() const override;
+	virtual bool IsCheckedOutOther(FString* Who = nullptr) const override;
+	virtual bool IsCurrent() const override;
+	virtual bool IsSourceControlled() const override;
+	virtual bool IsAdded() const override;
+	virtual bool IsDeleted() const override;
+	virtual bool IsIgnored() const override;
+	virtual bool CanEdit() const override;
+	virtual bool IsUnknown() const override;
+	virtual bool IsModified() const override;
+	virtual bool CanAdd() const override;
+	virtual bool IsConflicted() const override;
+
+public:
+	/** History of the item, if any */
+	TGitSourceControlHistory History;
+
+	/** Filename on disk */
+	FString LocalFilename;
+
+	/** State of the working copy */
+	EWorkingCopyState::Type WorkingCopyState;
+
+	/** The timestamp of the last update */
+	FDateTime TimeStamp;
+};
diff --git a/Engine/Plugins/Developer/GitSourceControl/Source/GitSourceControl/Private/GitSourceControlUtils.cpp b/Engine/Plugins/Developer/GitSourceControl/Source/GitSourceControl/Private/GitSourceControlUtils.cpp
new file mode 100644
--- /dev/null
+++ b/Engine/Plugins/Developer/GitSourceControl/Source/GitSourceControl/Private/GitSourceControlUtils.cpp
@@ -0,0 +1,694 @@
+// Copyright 1998-2015 Epic Games, Inc. All Rights Reserved.
+
+#include "GitSourceControlPrivatePCH.h"
+#include "GitSourceControlUtils.h"
+#include "GitSourceControlState.h"
+#include "GitSourceControlModule.h"
+#include "GitSourceControlCommand.h"
+
+#if PLATFORM_LINUX
+#include <sys/ioctl.h>
+#endif
+
+
+namespace GitSourceControlConstants
+{
+	/** The maximum number of files we submit in a single Git command */
+	const int32 MaxFilesPerBatch = 50;
+}
+
+FScopedTempFile::FScopedTempFile(const FText& InText)
+{
+	Filename = FPaths::CreateTempFilename(*FPaths::GameLogDir(), TEXT("Git-Temp"), TEXT(".txt"));
+	if (!FFileHelper::SaveStringToFile(InText.ToString(), *Filename, FFileHelper::EEncodingOptions::ForceUTF8WithoutBOM))
+	{
+		UE_LOG(LogSourceControl, Error, TEXT("Failed to write to temp file: %s"), *Filename);
+	}
+}
+
+FScopedTempFile::~FScopedTempFile()
+{
+	if(FPaths::FileExists(Filename))
+	{
+		if(!FPlatformFileManager::Get().GetPlatformFile().DeleteFile(*Filename))
+		{
+			UE_LOG(LogSourceControl, Error, TEXT("Failed to delete temp file: %s"), *Filename);
+		}
+	}
+}
+
+const FString& FScopedTempFile::GetFilename() const
+{
+	return Filename;
+}
+
+
+namespace GitSourceControlUtils
+{
+
+// Launch the Git command line process and extract its results & errors
+static bool RunCommandInternalRaw(const FString& InCommand, const FString& InPathToGitBinary, const FString& InRepositoryRoot, const TArray<FString>& InParameters, const TArray<FString>& InFiles, FString& OutResults, FString& OutErrors)
+{
+	int32 ReturnCode = 0;
+	FString FullCommand;
+	FString LogableCommand; // short version of the command for logging purpose
+
+	if (!InRepositoryRoot.IsEmpty())
+	{
+		// Specify the working copy (the root) of the git repository (before the command itself)
+		// @todo Does not work in UE4.1 on Mac if there is a space in the path ("[REDACTED_USER_HOME]/Unreal Project/MyProject")
+		FullCommand  = TEXT("--work-tree=\"");
+		FullCommand += InRepositoryRoot;
+		// and the ".git" subdirectory in it (before the command itself)
+		FullCommand += TEXT("\" --git-dir=\"");
+		FullCommand += InRepositoryRoot;
+		FullCommand += TEXT(".git\" ");
+	}
+	// then the git command itself ("status", "log", "commit"...)
+	LogableCommand += InCommand;
+
+	// Append to the command all parameters, and then finally the files
+	for(const auto& Parameter : InParameters)
+	{
+		LogableCommand += TEXT(" ");
+		LogableCommand += Parameter;
+	}
+	for(const auto& File : InFiles)
+	{
+		LogableCommand += TEXT(" \"");
+		LogableCommand += File;
+		LogableCommand += TEXT("\"");
+	}
+	// Also, Git does not have a "--non-interactive" option, as it auto-detects when there are no connected standard input/output streams
+
+	FullCommand += LogableCommand;
+
+	UE_LOG(LogSourceControl, Log, TEXT("RunCommandInternalRaw: 'git %s'"), *LogableCommand);
+// @todo: temporary debug logs
+//	UE_LOG(LogSourceControl, Log, TEXT("RunCommandInternalRaw: 'git %s'"), *FullCommand);
+	FPlatformProcess::ExecProcess(*InPathToGitBinary, *FullCommand, &ReturnCode, &OutResults, &OutErrors);
+	UE_LOG(LogSourceControl, Log, TEXT("RunCommandInternalRaw: ExecProcess ReturnCode=%d OutResults='%s'"), ReturnCode, *OutResults);
+	if(!OutErrors.IsEmpty())
+	{
+		UE_LOG(LogSourceControl, Error, TEXT("RunCommandInternalRaw: ExecProcess ReturnCode=%d OutErrors='%s'"), ReturnCode, *OutErrors);
+	}
+
+	return ReturnCode == 0;
+}
+
+// Basic parsing or results & errors from the Git command line process
+static bool RunCommandInternal(const FString& InCommand, const FString& InPathToGitBinary, const FString& InRepositoryRoot, const TArray<FString>& InParameters, const TArray<FString>& InFiles, TArray<FString>& OutResults, TArray<FString>& OutErrorMessages)
+{
+	bool bResult;
+	FString Results;
+	FString Errors;
+
+	bResult = RunCommandInternalRaw(InCommand, InPathToGitBinary, InRepositoryRoot, InParameters, InFiles, Results, Errors);
+	Results.ParseIntoArray(&OutResults, TEXT("\n"), true);
+	Errors.ParseIntoArray(&OutErrorMessages, TEXT("\n"), true);
+
+	return bResult;
+}
+
+FString FindGitBinaryPath()
+{
+	bool bFound = false;
+
+#if PLATFORM_WINDOWS
+	// 1) First of all, check for the ThirdParty directory as it may contain a specific version of Git for this plugin to work
+	// NOTE using only "git" (or "git.exe") relying on the "PATH" envvar does not always work as expected, depending on the installation:
+	// If the PATH is set with "git/cmd" instead of "git/bin",
+	// "git.exe" launch "git/cmd/git.exe" that redirect to "git/bin/git.exe" and ExecProcess() is unable to catch its outputs streams.
+
+	// Under Windows, we can use the third party "msysgit PortableGit" https://code.google.com/p/msysgit/downloads/list?can=1&q=PortableGit
+	// NOTE: Win32 platform subdirectory as there is no Git 64bit build available
+	FString GitBinaryPath(FPaths::EngineDir() / TEXT("Binaries/ThirdParty/git/Win32/bin") / TEXT("git.exe"));
+	bFound = CheckGitAvailability(GitBinaryPath);
+#else
+	FString GitBinaryPath;
+#endif
+
+	// 2) If Git is not found in ThirdParty directory, look into standard install directory
+	if(!bFound)
+	{
+#if PLATFORM_WINDOWS
+		// @todo use the Windows registry to find Git
+		GitBinaryPath = TEXT("C:/Program Files (x86)/Git/bin/git.exe");
+#else
+		GitBinaryPath = TEXT("/usr/bin/git");
+#endif
+	}
+
+	return GitBinaryPath;
+}
+
+bool CheckGitAvailability(const FString& InPathToGitBinary)
+{
+	bool bGitAvailable = false;
+
+	FString InfoMessages;
+	FString ErrorMessages;
+	bGitAvailable = RunCommandInternalRaw(TEXT("version"), InPathToGitBinary, FString(), TArray<FString>(), TArray<FString>(), InfoMessages, ErrorMessages);
+	if(bGitAvailable)
+	{
+		if(InfoMessages.Contains("git"))
+		{
+			bGitAvailable = true;
+		}
+		else
+		{
+			bGitAvailable = false;
+		}
+	}
+
+	// @todo also check Git config user.name & user.email
+
+	return bGitAvailable;
+}
+
+bool FindRootDirectory(const FString& InPathToGameDir, FString& OutRepositoryRoot)
+{
+	bool bFound = false;
+	FString PathToGitSubdirectory;
+	OutRepositoryRoot = InPathToGameDir;
+
+	while(!bFound && !OutRepositoryRoot.IsEmpty())
+	{
+		PathToGitSubdirectory = OutRepositoryRoot;
+		PathToGitSubdirectory += TEXT(".git"); // Look for the ".git" subdirectory at the root of every Git repository
+		bFound = IFileManager::Get().DirectoryExists(*PathToGitSubdirectory);
+		if(!bFound)
+		{
+			int32 LastSlashIndex;
+			OutRepositoryRoot = OutRepositoryRoot.LeftChop(5);
+			if(OutRepositoryRoot.FindLastChar('/', LastSlashIndex))
+			{
+				OutRepositoryRoot = OutRepositoryRoot.Left(LastSlashIndex + 1);
+			}
+			else
+			{
+				OutRepositoryRoot.Empty();
+			}
+		}
+	}
+
+	return bFound;
+}
+
+bool RunCommand(const FString& InCommand, const FString& InPathToGitBinary, const FString& InRepositoryRoot, const TArray<FString>& InParameters, const TArray<FString>& InFiles, TArray<FString>& OutResults, TArray<FString>& OutErrorMessages)
+{
+	bool bResult = true;
+
+	if(InFiles.Num() > GitSourceControlConstants::MaxFilesPerBatch)
+	{
+		// Batch files up so we dont exceed command-line limits
+		int32 FileCount = 0;
+		while(FileCount < InFiles.Num())
+		{
+			TArray<FString> FilesInBatch;
+			for(int32 FileIndex = 0; FileCount < InFiles.Num() && FileIndex < GitSourceControlConstants::MaxFilesPerBatch; FileIndex++, FileCount++)
+			{
+				FilesInBatch.Add(InFiles[FileCount]);
+			}
+
+			TArray<FString> BatchResults;
+			TArray<FString> BatchErrors;
+			bResult &= RunCommandInternal(InCommand, InPathToGitBinary, InRepositoryRoot, InParameters, FilesInBatch, BatchResults, BatchErrors);
+			OutResults += BatchResults;
+			OutErrorMessages += BatchErrors;
+		}
+	}
+	else
+	{
+		bResult &= RunCommandInternal(InCommand, InPathToGitBinary, InRepositoryRoot, InParameters, InFiles, OutResults, OutErrorMessages);
+	}
+
+	return bResult;
+}
+
+bool RunCommit(const FString& InPathToGitBinary, const FString& InRepositoryRoot, const TArray<FString>& InParameters, const TArray<FString>& InFiles, TArray<FString>& OutResults, TArray<FString>& OutErrorMessages)
+{
+	bool bResult = true;
+
+	if (InFiles.Num() > GitSourceControlConstants::MaxFilesPerBatch)
+	{
+		// Batch files up so we dont exceed command-line limits
+		int32 FileCount = 0;
+		TArray<FString> FilesInBatch;
+		for(int32 FileIndex = 0; FileIndex < GitSourceControlConstants::MaxFilesPerBatch; FileIndex++, FileCount++)
+		{
+			FilesInBatch.Add(InFiles[FileCount]);
+		}
+		// First batch is a simple "git commit" command with only the first files
+		bResult &= RunCommandInternal(TEXT("commit"), InPathToGitBinary, InRepositoryRoot, InParameters, InFiles, OutResults, OutErrorMessages);
+		
+		TArray<FString> Parameters;
+		for(const auto& Parameter : InParameters)
+		{
+			Parameters.Add(Parameter);
+		}
+		Parameters.Add(TEXT("--amend"));
+
+		while (FileCount < InFiles.Num())
+		{
+			TArray<FString> FilesInBatch;
+			for(int32 FileIndex = 0; FileCount < InFiles.Num() && FileIndex < GitSourceControlConstants::MaxFilesPerBatch; FileIndex++, FileCount++)
+			{
+				FilesInBatch.Add(InFiles[FileCount]);
+			}
+			// Next batches "amend" the commit with some more files
+			TArray<FString> BatchResults;
+			TArray<FString> BatchErrors;
+			bResult &= RunCommandInternal(TEXT("commit"), InPathToGitBinary, InRepositoryRoot, Parameters, FilesInBatch, BatchResults, BatchErrors);
+			OutResults += BatchResults;
+			OutErrorMessages += BatchErrors;
+		}
+	}
+	else
+	{
+		RunCommandInternal(TEXT("commit"), InPathToGitBinary, InRepositoryRoot, InParameters, InFiles, OutResults, OutErrorMessages);
+	}
+
+	return bResult;
+}
+
+/** Match the relative filename of a Git status result with a provided absolute filename */
+class FGitStatusFileMatcher
+{
+public:
+	FGitStatusFileMatcher(const FString& InAbsoluteFilename)
+		: AbsoluteFilename(InAbsoluteFilename)
+	{
+	}
+
+	bool operator()(const FString& InResult) const
+	{
+		// Extract the relative filename from the Git status result
+		FString RelativeFilename = InResult.RightChop(3);
+		// Note: this is not enough in case of a rename from -> to
+		int32 RenameIndex;
+		if (RelativeFilename.FindLastChar('>', RenameIndex))
+		{
+			// Extract only the second part of a rename "from -> to"
+			RelativeFilename = RelativeFilename.RightChop(RenameIndex + 2);
+		}
+		return AbsoluteFilename.Contains(RelativeFilename);
+	}
+
+private:
+	const FString& AbsoluteFilename;
+};
+
+/**
+* Extract and interpret the file state from the given Git status result.
+* @see http://git-scm.com/docs/git-status
+* ' ' = unmodified
+* 'M' = modified
+* 'A' = added
+* 'D' = deleted
+* 'R' = renamed
+* 'C' = copied
+* 'U' = updated but unmerged
+* '?' = unknown/untracked
+* '!' = ignored
+*/
+class FGitStatusParser
+{
+public:
+	FGitStatusParser(const FString& InResult)
+	{
+		TCHAR IndexState = InResult[0];
+		TCHAR WCopyState = InResult[1];
+		if(   (IndexState == 'U' || WCopyState == 'U')
+		   || (IndexState == 'A' && WCopyState == 'A')
+		   || (IndexState == 'D' && WCopyState == 'D'))
+		{
+			// "Unmerged" conflict cases are generally marked with a "U",
+			// but there are also the special cases of both "A"dded, or both "D"eleted
+			State = EWorkingCopyState::Conflicted;
+		}
+		else if(IndexState == 'A')
+		{
+			State = EWorkingCopyState::Added;
+		}
+		else if(IndexState == 'D')
+		{
+			State = EWorkingCopyState::Deleted;
+		}
+		else if(WCopyState == 'D')
+		{
+			State = EWorkingCopyState::Missing;
+		}
+		else if(IndexState == 'M' || WCopyState == 'M')
+		{
+			State = EWorkingCopyState::Modified;
+		}
+		else if(IndexState == 'R')
+		{
+			State = EWorkingCopyState::Renamed;
+		}
+		else if(IndexState == 'C')
+		{
+			State = EWorkingCopyState::Copied;
+		}
+		else if(IndexState == '?' || WCopyState == '?')
+		{
+			State = EWorkingCopyState::NotControlled;
+		}
+		else if(IndexState == '!' || WCopyState == '!')
+		{
+			State = EWorkingCopyState::Ignored;
+		}
+		else
+		{
+			// Unmodified never yield a status
+			State = EWorkingCopyState::Unknown;
+		}
+	}
+
+	EWorkingCopyState::Type State;
+};
+
+/** Parse the array of strings results of a 'git status' command
+*
+* Example git status results:
+M  Content/Textures/T_Perlin_Noise_M.uasset
+R  Content/Textures/T_Perlin_Noise_M.uasset -> Content/Textures/T_Perlin_Noise_M2.uasset
+?? Content/Materials/M_Basic_Wall.uasset
+!! BasicCode.sln
+*/
+static void ParseStatusResults(const TArray<FString>& InFiles, const TArray<FString>& InResults, TArray<FGitSourceControlState>& OutStates)
+{
+	// Iterate on all files explicitely listed in the command
+	for(const auto& File : InFiles)
+	{
+		FGitSourceControlState FileState(File);
+		// Search the file in the list of status
+		int32 IdxResult = InResults.IndexOfByPredicate(FGitStatusFileMatcher(File));
+		if(IdxResult != INDEX_NONE)
+		{
+			// File found in status results; only the case for "changed" files
+			FGitStatusParser StatusParser(InResults[IdxResult]);
+			FileState.WorkingCopyState = StatusParser.State;
+		}
+		else
+		{
+			// File not found in status
+			if(FPaths::FileExists(File))
+			{
+				// usually means the file is unchanged,
+				FileState.WorkingCopyState = EWorkingCopyState::Unchanged;
+			}
+			else
+			{
+				// but also the case for newly created content: there is no file on disk until the content is saved for the first time
+				FileState.WorkingCopyState = EWorkingCopyState::NotControlled;
+			}
+		}
+		FileState.TimeStamp.Now();
+		OutStates.Add(FileState);
+	}
+}
+
+bool RunUpdateStatus(const FString& InPathToGitBinary, const FString& InRepositoryRoot, const TArray<FString>& InFiles, TArray<FString>& OutErrorMessages, TArray<FGitSourceControlState>& OutStates)
+{
+	bool bResults = true;
+	TArray<FString> Results;
+	TArray<FString> Parameters;
+	Parameters.Add(TEXT("--porcelain"));
+	Parameters.Add(TEXT("--ignored"));
+
+	// Git status does not show any "untracked files" when called with files from different subdirectories! (issue #3)
+	// 1) So here we group files by path (ie. by subdirectory)
+	TMap<FString, TArray<FString>> GroupOfFiles;
+	for(const auto& File : InFiles)
+	{
+		const FString Path = FPaths::GetPath(*File);
+		TArray<FString>* Group = GroupOfFiles.Find(Path);
+		if(Group != nullptr)
+		{
+			Group->Add(File);
+		}
+		else
+		{
+			TArray<FString> Group;
+			Group.Add(File);
+			GroupOfFiles.Add(Path, Group);
+		}
+	}
+
+	// 2) then we can batch git status operation by subdirectory
+	for(const auto& Files : GroupOfFiles)
+	{
+		bool bResult = RunCommand(TEXT("status"), InPathToGitBinary, InRepositoryRoot, Parameters, Files.Value, Results, OutErrorMessages);
+		if(bResult)
+		{
+			ParseStatusResults(Files.Value, Results, OutStates);
+		}
+		else
+		{
+			bResults = false;
+		}
+	}
+
+	return bResults;
+}
+
+/**
+* Run a Git show command to dump the binary content of a revision into a file.
+*/
+bool RunDumpToFile(const FString& InPathToGitBinary, const FString& InRepositoryRoot, const FString& InParameter, const FString& InDumpFileName)
+{
+	bool bResult = false;
+	FString FullCommand;
+
+	if(!InRepositoryRoot.IsEmpty())
+	{
+		// Specify the working copy (the root) of the git repository (before the command itself)
+		FullCommand = TEXT("--work-tree=\"");
+		FullCommand += InRepositoryRoot;
+		// and the ".git" subdirectory in it (before the command itself)
+		FullCommand += TEXT("\" --git-dir=\"");
+		FullCommand += InRepositoryRoot;
+		FullCommand += TEXT(".git\" ");
+	}
+	// then the git command itself
+	FullCommand += TEXT("show ");
+
+	// Append to the command the parameter
+	FullCommand += InParameter;
+
+	const bool bLaunchDetached = false;
+	const bool bLaunchHidden = true;
+	const bool bLaunchReallyHidden = bLaunchHidden;
+
+	// Setup output redirection pipes, so that we can harvest compiler output and display it ourselves
+#if PLATFORM_LINUX
+	int pipefd[2];
+	pipe(pipefd);
+	void* PipeRead = &pipefd[0];
+	void* PipeWrite = &pipefd[1];
+#else
+	void* PipeRead = NULL;
+	void* PipeWrite = NULL;
+#endif
+
+	verify(FPlatformProcess::CreatePipe(PipeRead, PipeWrite));
+
+	// @todo temp debug log
+	//UE_LOG(LogSourceControl, Log, TEXT("RunDumpToFile: 'git %s'"), *FullCommand);
+	FProcHandle ProcessHandle = FPlatformProcess::CreateProc(*InPathToGitBinary, *FullCommand, bLaunchDetached, bLaunchHidden, bLaunchReallyHidden, NULL, 0, NULL, PipeWrite);
+	//UE_LOG(LogSourceControl, Log, TEXT("RunDumpToFile: ProcessHandle=%x"), ProcessHandle.Get());
+	if(ProcessHandle.IsValid())
+	{
+		FPlatformProcess::Sleep(0.01);
+
+		TArray<uint8> BinaryFileContent;
+		while(FPlatformProcess::IsProcRunning(ProcessHandle))
+		{
+			TArray<uint8> BinaryData;
+			FPlatformProcess::ReadPipeToArray(PipeRead, BinaryData);
+			if(BinaryData.Num() > 0)
+			{
+				BinaryFileContent.Append(MoveTemp(BinaryData));
+			}
+		}
+		TArray<uint8> BinaryData;
+		FPlatformProcess::ReadPipeToArray(PipeRead, BinaryData);
+		if(BinaryData.Num() > 0)
+		{
+			BinaryFileContent.Append(MoveTemp(BinaryData));
+		}
+		// Save buffer into temp file
+		if(FFileHelper::SaveArrayToFile(BinaryFileContent, *InDumpFileName))
+		{
+			UE_LOG(LogSourceControl, Log, TEXT("Writed '%s' (%do)"), *InDumpFileName, BinaryFileContent.Num());
+			bResult = true;
+		}
+		else
+		{
+			UE_LOG(LogSourceControl, Error, TEXT("Could not write %s"), *InDumpFileName);
+		}
+	}
+	else
+	{
+		UE_LOG(LogSourceControl, Error, TEXT("Failed to launch 'git show'"));
+	}
+
+#if PLATFORM_LINUX
+	close(*(int*)PipeRead);
+	close(*(int*)PipeWrite);
+#else
+	FPlatformProcess::ClosePipe(PipeRead, PipeWrite);
+#endif
+
+	return bResult;
+}
+
+
+/**
+* Extract and interpret the file state from the given Git log --name-status.
+* @see https://www.kernel.org/pub/software/scm/git/docs/git-log.html
+* ' ' = unmodified
+* 'M' = modified
+* 'A' = added
+* 'D' = deleted
+* 'R' = renamed
+* 'C' = copied
+* 'T' = type changed
+* 'U' = updated but unmerged
+* 'X' = unknown
+* 'B' = broken pairing
+*/
+FString LogStatusToString(TCHAR InStatus)
+{
+	switch(InStatus)
+	{
+	case TEXT(' '):
+		return FString("unmodified");
+	case TEXT('M'):
+		return FString("modified");
+	case TEXT('A'):
+		return FString("added");
+	case TEXT('D'):
+		return FString("deleted");
+	case TEXT('R'):
+		return FString("renamed");
+	case TEXT('C'):
+		return FString("copied");
+	case TEXT('T'):
+		return FString("type changed");
+	case TEXT('U'):
+		return FString("unmerged");
+	case TEXT('X'):
+		return FString("unknown");
+	case TEXT('B'):
+		return FString("broked pairing");
+	}
+
+	return FString();
+}
+
+/** Parse the array of strings results of a 'git log' command
+*
+* Example git log results:
+commit 97a4e7626681895e073aaefd68b8ac087db81b0b
+Author: [REDACTED_PERSON] <[REDACTED_EMAIL]>
+Date:   2014-2015-05-15 21:32:27 +0200
+
+    Another commit used to test History
+
+     - with many lines
+     - some <xml>
+     - and strange characteres $*+
+
+M   Content/Blueprints/Blueprint_CeilingLight.uasset
+R100    Content/Textures/T_Concrete_Poured_D.uasset Content/Textures/T_Concrete_Poured_D2.uasset
+
+commit 355f0df26ebd3888adbb558fd42bb8bd3e565000
+Author: [REDACTED_PERSON] <[REDACTED_EMAIL]>
+Date:   2014-2015-05-12 11:28:14 +0200
+
+    Testing git status, edit, and revert
+
+A    Content/Blueprints/Blueprint_CeilingLight.uasset
+C099    Content/Textures/T_Concrete_Poured_N.uasset Content/Textures/T_Concrete_Poured_N2.uasset
+*/
+void ParseLogResults(const TArray<FString>& InResults, TGitSourceControlHistory& OutHistory)
+{
+	TSharedRef<FGitSourceControlRevision, ESPMode::ThreadSafe> SourceControlRevision = MakeShareable(new FGitSourceControlRevision);
+	for(const auto& Result : InResults)
+	{
+		if(Result.StartsWith(TEXT("commit "))) // Start of a new commit
+		{
+			// End of the previous commit
+			if(SourceControlRevision->RevisionNumber != 0)
+			{
+				SourceControlRevision->Description += TEXT("\nCommit Id: ");
+				SourceControlRevision->Description += SourceControlRevision->CommitId;
+				OutHistory.Add(SourceControlRevision);
+
+				SourceControlRevision = MakeShareable(new FGitSourceControlRevision);
+			}
+			SourceControlRevision->CommitId = Result.RightChop(7);
+			FString ShortCommitId = SourceControlRevision->CommitId.Right(8); // Short revision ; first 8 hex characters (max that can hold a 32 bit integer)
+			SourceControlRevision->RevisionNumber = FParse::HexNumber(*ShortCommitId);
+		}
+		else if(Result.StartsWith(TEXT("Author: "))) // Author name & email
+		{
+			// Remove the 'email' part of the UserName
+			FString UserNameEmail = Result.RightChop(8);
+			int32 EmailIndex = 0;
+			if(UserNameEmail.FindLastChar('<', EmailIndex))
+			{
+				SourceControlRevision->UserName = UserNameEmail.Left(EmailIndex - 1);
+			}
+		}
+		else if(Result.StartsWith(TEXT("Date:   "))) // Commit date
+		{
+			FString Date = Result.RightChop(8);
+			SourceControlRevision->Date = FDateTime::FromUnixTimestamp(FCString::Atoi(*Date));
+		}
+	//	else if(Result.IsEmpty()) // empty line before/after commit message has already been taken care by FString::ParseIntoArray()
+		else if(Result.StartsWith(TEXT("    ")))  // Multi-lines commit message
+		{
+			SourceControlRevision->Description += Result.RightChop(4);
+			SourceControlRevision->Description += TEXT("\n");
+		}
+		else // List of modified files, starting with a uppercase status letter ("A"/"M"...)
+		{
+			TCHAR Status = Result[0];
+			SourceControlRevision->Action = LogStatusToString(Status); // Readable action string ("Added", Modified"...) instead of "A"/"M"...
+			SourceControlRevision->Filename = Result.RightChop(2); // relative filename
+		}
+	}
+	// End of the last commit
+	if(SourceControlRevision->RevisionNumber != 0)
+	{
+		SourceControlRevision->Description += TEXT("\nCommit Id: ");
+		SourceControlRevision->Description += SourceControlRevision->CommitId;
+		OutHistory.Add(SourceControlRevision);
+	}
+}
+
+bool UpdateCachedStates(const TArray<FGitSourceControlState>& InStates)
+{
+	FGitSourceControlModule& GitSourceControl = FModuleManager::LoadModuleChecked<FGitSourceControlModule>( "GitSourceControl" );
+	FGitSourceControlProvider& Provider = GitSourceControl.GetProvider();
+	int NbStatesUpdated = 0;
+
+	for(const auto& InState : InStates)
+	{
+		TSharedRef<FGitSourceControlState, ESPMode::ThreadSafe> State = Provider.GetStateInternal(InState.LocalFilename);
+		if(State->WorkingCopyState != InState.WorkingCopyState)
+		{
+			State->WorkingCopyState = InState.WorkingCopyState;
+		//	State->TimeStamp = InState.TimeStamp; // @todo Bug report: Workaround a bug with the Source Control Module not updating file state after a "Save"
+			NbStatesUpdated++;
+		}
+	}
+
+	return (NbStatesUpdated > 0);
+}
+
+}
diff --git a/Engine/Plugins/Developer/GitSourceControl/Source/GitSourceControl/Private/GitSourceControlUtils.h b/Engine/Plugins/Developer/GitSourceControl/Source/GitSourceControl/Private/GitSourceControlUtils.h
new file mode 100644
--- /dev/null
+++ b/Engine/Plugins/Developer/GitSourceControl/Source/GitSourceControl/Private/GitSourceControlUtils.h
@@ -0,0 +1,114 @@
+// Copyright 1998-2015 Epic Games, Inc. All Rights Reserved.
+
+#pragma once
+
+#include "GitSourceControlState.h"
+#include "GitSourceControlRevision.h"
+
+/**
+ * Helper struct for maintaining temporary files for passing to commands
+ */
+class FScopedTempFile
+{
+public:
+
+	/** Constructor - open & write string to temp file */
+	FScopedTempFile(const FText& InText);
+
+	/** Destructor - delete temp file */
+	~FScopedTempFile();
+
+	/** Get the filename of this temp file - empty if it failed to be created */
+	const FString& GetFilename() const;
+
+private:
+	/** The filename we are writing to */
+	FString Filename;
+};
+
+namespace GitSourceControlUtils
+{
+
+/**
+ * Find the path to the Git binary, looking in a few standard place (ThirdParty subdirectory and standard sSystem paths)
+ * @returns the path to the Git binary if found, or the last path tested if no git found.
+ */
+FString FindGitBinaryPath();
+
+/**
+ * Run a Git "version" command to check the availability of the binary.
+ * @param InPathToGitBinary		The path to the Git binary
+ * @returns true if the command succeeded and returned no errors
+ */
+bool CheckGitAvailability(const FString& InPathToGitBinary);
+
+/**
+ * Find the root of the Git repository, looking from the GameDir and upward in its parent directories
+ * @param InPathToGameDir		The path to the Game Directory
+ * @param OutRepositoryRoot		The path to the root directory of the Git repository if found
+ * @returns true if the command succeeded and returned no errors
+*/
+bool FindRootDirectory(const FString& InPathToGameDir, FString& OutRepositoryRoot);
+
+/**
+ * Run a Git command - output is a string TArray.
+ *
+ * @param	InCommand			The Git command - e.g. commit
+ * @param	InPathToGitBinary	The path to the Git binary
+ * @param	InRepositoryRoot	The Git repository from where to run the command - usually the Game directory (can be empty)
+ * @param	InParameters		The parameters to the Git command
+ * @param	InFiles				The files to be operated on
+ * @param	OutResults			The results (from StdOut) as an array per-line
+ * @param	OutErrorMessages	Any errors (from StdErr) as an array per-line
+ * @returns true if the command succeeded and returned no errors
+ */
+bool RunCommand(const FString& InCommand, const FString& InPathToGitBinary, const FString& InRepositoryRoot, const TArray<FString>& InParameters, const TArray<FString>& InFiles, TArray<FString>& OutResults, TArray<FString>& OutErrorMessages);
+
+/**
+ * Run a Git commit command by batches.
+ *
+ * @param	InPathToGitBinary	The path to the Git binary
+ * @param	InRepositoryRoot	The Git repository from where to run the command - usually the Game directory (can be empty)
+ * @param	InParameter			The parameters to the Git commit command
+ * @param	InFiles				The files to be operated on
+ * @param	OutErrorMessages	Any errors (from StdErr) as an array per-line
+ * @returns true if the command succeeded and returned no errors
+ */
+bool RunCommit(const FString& InPathToGitBinary, const FString& InRepositoryRoot, const TArray<FString>& InParameters, const TArray<FString>& InFiles, TArray<FString>& OutResults, TArray<FString>& OutErrorMessages);
+
+/**
+ * Run a Git status command and parse it.
+ *
+ * @param	InPathToGitBinary	The path to the Git binary
+ * @param	InRepositoryRoot	The Git repository from where to run the command - usually the Game directory (can be empty)
+ * @param	InFiles				The files to be operated on
+ * @param	OutErrorMessages	Any errors (from StdErr) as an array per-line
+ * @returns true if the command succeeded and returned no errors
+ */
+bool RunUpdateStatus(const FString& InPathToGitBinary, const FString& InRepositoryRoot, const TArray<FString>& InFiles, TArray<FString>& OutErrorMessages, TArray<FGitSourceControlState>& OutStates);
+
+/**
+ * Run a Git show command to dump the binary content of a revision into a file.
+ *
+ * @param	InPathToGitBinary	The path to the Git binary
+ * @param	InRepositoryRoot	The Git repository from where to run the command - usually the Game directory (can be empty)
+ * @param	InParameter			The parameters to the Git show command (rev:path)
+ * @param	InDumpFileName		The temporary file to dump the revision
+ * @returns true if the command succeeded and returned no errors
+*/
+bool RunDumpToFile(const FString& InPathToGitBinary, const FString& InRepositoryRoot, const FString& InParameter, const FString& InDumpFileName);
+
+/**
+ * Parse the array of strings results of a 'git log' command
+ * @param	InResults			The results (from StdOut) as an array per-line
+ * @param	OutHistory			The history of the file
+ */
+void ParseLogResults(const TArray<FString>& InResults, TGitSourceControlHistory& OutHistory);
+
+/**
+ * Helper function for various commands to update cached states.
+ * @returns true if any states were updated
+ */
+bool UpdateCachedStates(const TArray<FGitSourceControlState>& InStates);
+
+}
diff --git a/Engine/Plugins/Developer/GitSourceControl/Source/GitSourceControl/Private/IGitSourceControlWorker.h b/Engine/Plugins/Developer/GitSourceControl/Source/GitSourceControl/Private/IGitSourceControlWorker.h
new file mode 100644
--- /dev/null
+++ b/Engine/Plugins/Developer/GitSourceControl/Source/GitSourceControl/Private/IGitSourceControlWorker.h
@@ -0,0 +1,25 @@
+// Copyright 1998-2015 Epic Games, Inc. All Rights Reserved.
+
+#pragma once
+
+class IGitSourceControlWorker
+{
+public:
+	/**
+	 * Name describing the work that this worker does. Used for factory method hookup.
+	 */
+	virtual FName GetName() const = 0;
+
+	/**
+	 * Function that actually does the work. Can be executed on another thread.
+	 */
+	virtual bool Execute( class FGitSourceControlCommand& InCommand ) = 0;
+
+	/**
+	 * Updates the state of any items after completion (if necessary). This is always executed on the main thread.
+	 * @returns true if states were updated
+	 */
+	virtual bool UpdateStates() const = 0;
+};
+
+typedef TSharedRef<IGitSourceControlWorker, ESPMode::ThreadSafe> FGitSourceControlWorkerRef;
\ No newline at end of file
diff --git a/Engine/Plugins/Developer/GitSourceControl/Source/GitSourceControl/Private/SGitSourceControlSettings.cpp b/Engine/Plugins/Developer/GitSourceControl/Source/GitSourceControl/Private/SGitSourceControlSettings.cpp
new file mode 100644
--- /dev/null
+++ b/Engine/Plugins/Developer/GitSourceControl/Source/GitSourceControl/Private/SGitSourceControlSettings.cpp
@@ -0,0 +1,75 @@
+// Copyright 1998-2015 Epic Games, Inc. All Rights Reserved.
+
+#include "GitSourceControlPrivatePCH.h"
+#include "SGitSourceControlSettings.h"
+#include "GitSourceControlModule.h"
+
+#define LOCTEXT_NAMESPACE "SGitSourceControlSettings"
+
+void SGitSourceControlSettings::Construct(const FArguments& InArgs)
+{
+	FSlateFontInfo Font = FEditorStyle::GetFontStyle(TEXT("SourceControl.LoginWindow.Font"));
+
+	ChildSlot
+	[
+		SNew( SBorder )
+		.BorderImage( FEditorStyle::GetBrush("DetailsView.CategoryBottom") )
+		.Padding( FMargin( 0.0f, 3.0f, 0.0f, 0.0f ) )
+		[
+			SNew(SVerticalBox)
+			+ SVerticalBox::Slot()
+			.FillHeight(1.0f)
+			.Padding(2.0f)
+			.VAlign(VAlign_Center)
+			[
+				SNew(SHorizontalBox)
+				+SHorizontalBox::Slot()
+				.FillWidth(1.0f)
+				[
+					SNew(SVerticalBox)
+					+SVerticalBox::Slot()
+					.FillHeight(1.0f)
+					.Padding(2.0f)
+					.VAlign(VAlign_Center)
+					[
+						SNew(STextBlock)
+						.Text(LOCTEXT("BinaryPathLabel", "Git Path"))
+						.ToolTipText(LOCTEXT("BinaryPathLabel_Tooltip", "Path to Git binary"))
+						.Font(Font)
+					]
+				]
+				+SHorizontalBox::Slot()
+				.FillWidth(2.0f)
+				[
+					SNew(SVerticalBox)
+					+SVerticalBox::Slot()
+					.FillHeight(1.0f)
+					.Padding(2.0f)
+					[
+						SNew(SEditableTextBox)
+						.Text(this, &SGitSourceControlSettings::GetBinaryPathText)
+						.ToolTipText(LOCTEXT("BinaryPathLabel_Tooltip", "Path to Git binary"))
+						.OnTextCommitted(this, &SGitSourceControlSettings::OnBinaryPathTextCommited)
+						.OnTextChanged(this, &SGitSourceControlSettings::OnBinaryPathTextCommited, ETextCommit::Default)
+						.Font(Font)
+					]
+				]
+			]
+		]
+	];
+}
+
+FText SGitSourceControlSettings::GetBinaryPathText() const
+{
+	FGitSourceControlModule& GitSourceControl = FModuleManager::LoadModuleChecked<FGitSourceControlModule>("GitSourceControl");
+	return FText::FromString(GitSourceControl.AccessSettings().GetBinaryPath());
+}
+
+void SGitSourceControlSettings::OnBinaryPathTextCommited(const FText& InText, ETextCommit::Type InCommitType) const
+{
+	FGitSourceControlModule& GitSourceControl = FModuleManager::LoadModuleChecked<FGitSourceControlModule>("GitSourceControl");
+	GitSourceControl.AccessSettings().SetBinaryPath(InText.ToString());
+	GitSourceControl.SaveSettings();
+}
+
+#undef LOCTEXT_NAMESPACE
\ No newline at end of file
diff --git a/Engine/Plugins/Developer/GitSourceControl/Source/GitSourceControl/Private/SGitSourceControlSettings.h b/Engine/Plugins/Developer/GitSourceControl/Source/GitSourceControl/Private/SGitSourceControlSettings.h
new file mode 100644
--- /dev/null
+++ b/Engine/Plugins/Developer/GitSourceControl/Source/GitSourceControl/Private/SGitSourceControlSettings.h
@@ -0,0 +1,24 @@
+// Copyright 1998-2015 Epic Games, Inc. All Rights Reserved.
+
+#pragma once
+
+class SGitSourceControlSettings : public SCompoundWidget
+{
+public:
+	
+	SLATE_BEGIN_ARGS(SGitSourceControlSettings) {}
+	
+	SLATE_END_ARGS()
+
+public:
+
+	void Construct(const FArguments& InArgs);
+
+private:
+
+	/** Delegate to get binary path from settings */
+	FText GetBinaryPathText() const;
+
+	/** Delegate to commit repository text to settings */
+	void OnBinaryPathTextCommited(const FText& InText, ETextCommit::Type InCommitType) const;
+};
\ No newline at end of file
