Git source control plugin is now experimental

Added setting and wrapped registration of modular feature in the flag.

diff --git a/Engine/Plugins/Developer/GitSourceControl/Source/GitSourceControl/GitSourceControl.Build.cs b/Engine/Plugins/Developer/GitSourceControl/Source/GitSourceControl/GitSourceControl.Build.cs
--- a/Engine/Plugins/Developer/GitSourceControl/Source/GitSourceControl/GitSourceControl.Build.cs
+++ b/Engine/Plugins/Developer/GitSourceControl/Source/GitSourceControl/GitSourceControl.Build.cs
@@ -15,5 +15,12 @@ public class GitSourceControl : ModuleRules
 				"SourceControl",
 			}
 		);
+
+		if (UEBuildConfiguration.bBuildEditor == true)
+		{
+			// needed to enable/disable this via experimental settings
+			PrivateDependencyModuleNames.Add("CoreUObject");
+			PrivateDependencyModuleNames.Add("UnrealEd");
+		}
 	}
 }
diff --git a/Engine/Plugins/Developer/GitSourceControl/Source/GitSourceControl/Private/GitSourceControlModule.cpp b/Engine/Plugins/Developer/GitSourceControl/Source/GitSourceControl/Private/GitSourceControlModule.cpp
--- a/Engine/Plugins/Developer/GitSourceControl/Source/GitSourceControl/Private/GitSourceControlModule.cpp
+++ b/Engine/Plugins/Developer/GitSourceControl/Source/GitSourceControl/Private/GitSourceControlModule.cpp
@@ -8,6 +8,10 @@
 #include "GitSourceControlOperations.h"
 #include "Runtime/Core/Public/Features/IModularFeatures.h"
 
+#if WITH_EDITOR
+#include "Editor/UnrealEd/Classes/Settings/EditorExperimentalSettings.h"
+#endif
+
 #define LOCTEXT_NAMESPACE "GitSourceControl"
 
 template<typename Type>
@@ -34,8 +38,13 @@ void FGitSourceControlModule::StartupModule()
 	// load our settings
 	GitSourceControlSettings.LoadSettings();
 
-	// Bind our source control provider to the editor
-	IModularFeatures::Get().RegisterModularFeature( "SourceControl", &GitSourceControlProvider );
+#if WITH_EDITOR
+	if(GetDefault<UEditorExperimentalSettings>()->bGitSourceControl)
+	{
+		// Bind our source control provider to the editor
+		IModularFeatures::Get().RegisterModularFeature( "SourceControl", &GitSourceControlProvider );
+	}
+#endif
 }
 
 void FGitSourceControlModule::ShutdownModule()
diff --git a/Engine/Source/Editor/UnrealEd/Classes/Settings/EditorExperimentalSettings.h b/Engine/Source/Editor/UnrealEd/Classes/Settings/EditorExperimentalSettings.h
--- a/Engine/Source/Editor/UnrealEd/Classes/Settings/EditorExperimentalSettings.h
+++ b/Engine/Source/Editor/UnrealEd/Classes/Settings/EditorExperimentalSettings.h
@@ -85,6 +85,10 @@ public:
 	UPROPERTY(EditAnywhere, config, Category = AI, meta = (DisplayName = "Environment Querying System"))
 	bool bEQSEditor;
 	
+	/** Enables Git source control plugin */
+	UPROPERTY(EditAnywhere, config, Category = "Source Control", meta = (DisplayName = "Git Source Control plugin (requires editor restart)"))
+	bool bGitSourceControl;
+
 	/**
 	 * Returns an event delegate that is executed when a setting has changed.
 	 *
