Fixed Git 'Submit to source control' always failing

Added RemoveRedundantErrors (similar to SVN and P4) so thats erors we dont care about can be rectified.

diff --git a/Engine/Plugins/Developer/GitSourceControl/Source/GitSourceControl/Private/GitSourceControlOperations.cpp b/Engine/Plugins/Developer/GitSourceControl/Source/GitSourceControl/Private/GitSourceControlOperations.cpp
--- a/Engine/Plugins/Developer/GitSourceControl/Source/GitSourceControl/Private/GitSourceControlOperations.cpp
+++ b/Engine/Plugins/Developer/GitSourceControl/Source/GitSourceControl/Private/GitSourceControlOperations.cpp
@@ -188,6 +188,7 @@ bool FGitUpdateStatusWorker::Execute(FGitSourceControlCommand& InCommand)
 	if(InCommand.Files.Num() > 0)
 	{
 		InCommand.bCommandSuccessful = GitSourceControlUtils::RunUpdateStatus(InCommand.PathToGitBinary, InCommand.PathToRepositoryRoot, InCommand.Files, InCommand.ErrorMessages, States);
+		GitSourceControlUtils::RemoveRedundantErrors(InCommand, TEXT("' is outside repository"));
 
 		if(Operation->ShouldUpdateHistory())
 		{
diff --git a/Engine/Plugins/Developer/GitSourceControl/Source/GitSourceControl/Private/GitSourceControlUtils.cpp b/Engine/Plugins/Developer/GitSourceControl/Source/GitSourceControl/Private/GitSourceControlUtils.cpp
--- a/Engine/Plugins/Developer/GitSourceControl/Source/GitSourceControl/Private/GitSourceControlUtils.cpp
+++ b/Engine/Plugins/Developer/GitSourceControl/Source/GitSourceControl/Private/GitSourceControlUtils.cpp
@@ -440,7 +440,9 @@ bool RunUpdateStatus(const FString& InPathToGitBinary, const FString& InReposito
 	// 2) then we can batch git status operation by subdirectory
 	for(const auto& Files : GroupOfFiles)
 	{
-		bool bResult = RunCommand(TEXT("status"), InPathToGitBinary, InRepositoryRoot, Parameters, Files.Value, Results, OutErrorMessages);
+		TArray<FString> ErrorMessages;
+		bool bResult = RunCommand(TEXT("status"), InPathToGitBinary, InRepositoryRoot, Parameters, Files.Value, Results, ErrorMessages);
+		OutErrorMessages.Append(ErrorMessages);
 		if(bResult)
 		{
 			ParseStatusResults(Files.Value, Results, OutStates);
@@ -691,4 +693,49 @@ bool UpdateCachedStates(const TArray<FGitSourceControlState>& InStates)
 	return (NbStatesUpdated > 0);
 }
 
+/**
+ * Helper struct for RemoveRedundantErrors()
+ */
+struct FRemoveRedundantErrors
+{
+	FRemoveRedundantErrors(const FString& InFilter)
+		: Filter(InFilter)
+	{
+	}
+
+	bool operator()(const FString& String) const
+	{
+		if(String.Contains(Filter))
+		{
+			return true;
+		}
+
+		return false;
+	}
+
+	/** The filter string we try to identify in the reported error */
+	FString Filter;
+};
+
+void RemoveRedundantErrors(FGitSourceControlCommand& InCommand, const FString& InFilter)
+{
+	bool bFoundRedundantError = false;
+	for(auto Iter(InCommand.ErrorMessages.CreateConstIterator()); Iter; Iter++)
+	{
+		if(Iter->Contains(InFilter))
+		{
+			InCommand.InfoMessages.Add(*Iter);
+			bFoundRedundantError = true;
+		}
+	}
+
+	InCommand.ErrorMessages.RemoveAll( FRemoveRedundantErrors(InFilter) );
+
+	// if we have no error messages now, assume success!
+	if(bFoundRedundantError && InCommand.ErrorMessages.Num() == 0 && !InCommand.bCommandSuccessful)
+	{
+		InCommand.bCommandSuccessful = true;
+	}
+}
+
 }
diff --git a/Engine/Plugins/Developer/GitSourceControl/Source/GitSourceControl/Private/GitSourceControlUtils.h b/Engine/Plugins/Developer/GitSourceControl/Source/GitSourceControl/Private/GitSourceControlUtils.h
--- a/Engine/Plugins/Developer/GitSourceControl/Source/GitSourceControl/Private/GitSourceControlUtils.h
+++ b/Engine/Plugins/Developer/GitSourceControl/Source/GitSourceControl/Private/GitSourceControlUtils.h
@@ -4,6 +4,7 @@
 
 #include "GitSourceControlState.h"
 #include "GitSourceControlRevision.h"
+#include "GitSourceControlCommand.h"
 
 /**
  * Helper struct for maintaining temporary files for passing to commands
@@ -111,4 +112,10 @@ void ParseLogResults(const TArray<FString>& InResults, TGitSourceControlHistory&
  */
 bool UpdateCachedStates(const TArray<FGitSourceControlState>& InStates);
 
+/** 
+ * Remove redundant errors (that contain a particular string) and also
+ * update the commands success status if all errors were removed.
+ */
+void RemoveRedundantErrors(FGitSourceControlCommand& InCommand, const FString& InFilter);
+
 }
