[INTEGRATE] Change 2412228 by [REDACTED_PERSON][REDACTED_HANDLE] on 2015/01/20 10:40:35

	Fix compile errors

diff --git a/Engine/Plugins/Developer/GitSourceControl/Source/GitSourceControl/Private/GitSourceControlCommand.h b/Engine/Plugins/Developer/GitSourceControl/Source/GitSourceControl/Private/GitSourceControlCommand.h
--- a/Engine/Plugins/Developer/GitSourceControl/Source/GitSourceControl/Private/GitSourceControlCommand.h
+++ b/Engine/Plugins/Developer/GitSourceControl/Source/GitSourceControl/Private/GitSourceControlCommand.h
@@ -5,7 +5,7 @@
 /**
  * Used to execute Git commands multi-threaded.
  */
-class FGitSourceControlCommand : public IQueuedWork
+class FGitSourceControlCommand : public FQueuedWork
 {
 public:
 
diff --git a/Engine/Source/Editor/BehaviorTreeEditor/Private/SBehaviorTreeDiff.cpp b/Engine/Source/Editor/BehaviorTreeEditor/Private/SBehaviorTreeDiff.cpp
--- a/Engine/Source/Editor/BehaviorTreeEditor/Private/SBehaviorTreeDiff.cpp
+++ b/Engine/Source/Editor/BehaviorTreeEditor/Private/SBehaviorTreeDiff.cpp
@@ -7,6 +7,7 @@
 #include "Editor/PropertyEditor/Public/IDetailsView.h"
 #include "BehaviorTree/BehaviorTree.h"
 #include "GenericCommands.h"
+#include "ISourceControlModule.h"
 
 #define LOCTEXT_NAMESPACE "SBehaviorTreeDiff"
 
diff --git a/Engine/Source/Editor/Kismet/Private/DiffUtils.cpp b/Engine/Source/Editor/Kismet/Private/DiffUtils.cpp
--- a/Engine/Source/Editor/Kismet/Private/DiffUtils.cpp
+++ b/Engine/Source/Editor/Kismet/Private/DiffUtils.cpp
@@ -8,6 +8,7 @@
 #include "Engine/SimpleConstructionScript.h"
 #include "IAssetTypeActions.h"
 #include "ObjectEditorUtils.h"
+#include "ISourceControlModule.h"
 
 static const UProperty* Resolve( const UStruct* Class, FName PropertyName )
 {
