Fixed deprecation warning

diff --git a/Engine/Plugins/Developer/GitSourceControl/Source/GitSourceControl/Private/GitSourceControlCommand.h b/Engine/Plugins/Developer/GitSourceControl/Source/GitSourceControl/Private/GitSourceControlCommand.h
--- a/Engine/Plugins/Developer/GitSourceControl/Source/GitSourceControl/Private/GitSourceControlCommand.h
+++ b/Engine/Plugins/Developer/GitSourceControl/Source/GitSourceControl/Private/GitSourceControlCommand.h
@@ -5,7 +5,7 @@
 /**
  * Used to execute Git commands multi-threaded.
  */
-class FGitSourceControlCommand : public FQueuedWork
+class FGitSourceControlCommand : public IQueuedWork
 {
 public:
 
