Fix for double pipe calls in DesktopPlatform/GitSCC.

- PR #873 contributed by [REDACTED_PERSON]

https://github.com/EpicGames/UnrealEngine/pull/873

diff --git a/Engine/Plugins/Developer/GitSourceControl/Source/GitSourceControl/Private/GitSourceControlUtils.cpp b/Engine/Plugins/Developer/GitSourceControl/Source/GitSourceControl/Private/GitSourceControlUtils.cpp
--- a/Engine/Plugins/Developer/GitSourceControl/Source/GitSourceControl/Private/GitSourceControlUtils.cpp
+++ b/Engine/Plugins/Developer/GitSourceControl/Source/GitSourceControl/Private/GitSourceControlUtils.cpp
@@ -487,16 +487,8 @@ bool RunDumpToFile(const FString& InPathToGitBinary, const FString& InRepository
 	const bool bLaunchHidden = true;
 	const bool bLaunchReallyHidden = bLaunchHidden;
 
-	// Setup output redirection pipes, so that we can harvest compiler output and display it ourselves
-#if PLATFORM_LINUX
-	int pipefd[2];
-	pipe(pipefd);
-	void* PipeRead = &pipefd[0];
-	void* PipeWrite = &pipefd[1];
-#else
 	void* PipeRead = NULL;
 	void* PipeWrite = NULL;
-#endif
 
 	verify(FPlatformProcess::CreatePipe(PipeRead, PipeWrite));
 
@@ -540,12 +532,7 @@ bool RunDumpToFile(const FString& InPathToGitBinary, const FString& InRepository
 		UE_LOG(LogSourceControl, Error, TEXT("Failed to launch 'git show'"));
 	}
 
-#if PLATFORM_LINUX
-	close(*(int*)PipeRead);
-	close(*(int*)PipeWrite);
-#else
 	FPlatformProcess::ClosePipe(PipeRead, PipeWrite);
-#endif
 
 	return bResult;
 }
diff --git a/Engine/Source/Developer/DesktopPlatform/Private/DesktopPlatformBase.cpp b/Engine/Source/Developer/DesktopPlatform/Private/DesktopPlatformBase.cpp
--- a/Engine/Source/Developer/DesktopPlatform/Private/DesktopPlatformBase.cpp
+++ b/Engine/Source/Developer/DesktopPlatform/Private/DesktopPlatformBase.cpp
@@ -540,16 +540,8 @@ bool FDesktopPlatformBase::IsUnrealBuildToolAvailable()
 
 bool FDesktopPlatformBase::InvokeUnrealBuildToolSync(const FString& InCmdLineParams, FOutputDevice &Ar, bool bSkipBuildUBT, int32& OutReturnCode, FString& OutProcOutput)
 {
-	// Setup output redirection pipes, so that we can harvest compiler output and display it ourselves
-#if PLATFORM_LINUX
-	int pipefd[2];
-	pipe(pipefd);
-	void* PipeRead = &pipefd[0];
-	void* PipeWrite = &pipefd[1];
-#else
-	void* PipeRead = NULL;
-	void* PipeWrite = NULL;
-#endif
+	void* PipeRead = nullptr;
+	void* PipeWrite = nullptr;
 
 	verify(FPlatformProcess::CreatePipe(PipeRead, PipeWrite));
 
