PR #956: (Git Plugin) Search for the git executable into the GitHub folder if not found as standalone installation path (Contributed by [REDACTED_PERSON])

diff --git a/Engine/Plugins/Developer/GitSourceControl/Source/GitSourceControl/Private/GitSourceControlProvider.cpp b/Engine/Plugins/Developer/GitSourceControl/Source/GitSourceControl/Private/GitSourceControlProvider.cpp
--- a/Engine/Plugins/Developer/GitSourceControl/Source/GitSourceControl/Private/GitSourceControlProvider.cpp
+++ b/Engine/Plugins/Developer/GitSourceControl/Source/GitSourceControl/Private/GitSourceControlProvider.cpp
@@ -33,7 +33,14 @@ void FGitSourceControlProvider::CheckGitAvailability()
 {
 	FGitSourceControlModule& GitSourceControl = FModuleManager::LoadModuleChecked<FGitSourceControlModule>("GitSourceControl");
 	const FString& PathToGitBinary = GitSourceControl.AccessSettings().GetBinaryPath();
-	bGitAvailable = GitSourceControlUtils::CheckGitAvailability(PathToGitBinary);
+	if(!PathToGitBinary.IsEmpty())
+	{
+		bGitAvailable = GitSourceControlUtils::CheckGitAvailability(PathToGitBinary);
+	}
+	else
+	{
+		UE_LOG(LogSourceControl, Error, TEXT("Git not found"), *PathToGitBinary);
+	}
 }
 
 void FGitSourceControlProvider::Close()
diff --git a/Engine/Plugins/Developer/GitSourceControl/Source/GitSourceControl/Private/GitSourceControlSettings.cpp b/Engine/Plugins/Developer/GitSourceControl/Source/GitSourceControl/Private/GitSourceControlSettings.cpp
--- a/Engine/Plugins/Developer/GitSourceControl/Source/GitSourceControl/Private/GitSourceControlSettings.cpp
+++ b/Engine/Plugins/Developer/GitSourceControl/Source/GitSourceControl/Private/GitSourceControlSettings.cpp
@@ -32,7 +32,7 @@ void FGitSourceControlSettings::LoadSettings()
 	FScopeLock ScopeLock(&CriticalSection);
 	const FString& IniFile = SourceControlHelpers::GetSettingsIni();
 	bool bLoaded = GConfig->GetString(*GitSettingsConstants::SettingsSection, TEXT("BinaryPath"), BinaryPath, IniFile);
-	if(!bLoaded)
+	if(!bLoaded || BinaryPath.IsEmpty())
 	{
 		BinaryPath = GitSourceControlUtils::FindGitBinaryPath();
 	}
diff --git a/Engine/Plugins/Developer/GitSourceControl/Source/GitSourceControl/Private/GitSourceControlUtils.cpp b/Engine/Plugins/Developer/GitSourceControl/Source/GitSourceControl/Private/GitSourceControlUtils.cpp
--- a/Engine/Plugins/Developer/GitSourceControl/Source/GitSourceControl/Private/GitSourceControlUtils.cpp
+++ b/Engine/Plugins/Developer/GitSourceControl/Source/GitSourceControl/Private/GitSourceControlUtils.cpp
@@ -20,7 +20,7 @@ namespace GitSourceControlConstants
 FScopedTempFile::FScopedTempFile(const FText& InText)
 {
 	Filename = FPaths::CreateTempFilename(*FPaths::GameLogDir(), TEXT("Git-Temp"), TEXT(".txt"));
-	if (!FFileHelper::SaveStringToFile(InText.ToString(), *Filename, FFileHelper::EEncodingOptions::ForceUTF8WithoutBOM))
+	if(!FFileHelper::SaveStringToFile(InText.ToString(), *Filename, FFileHelper::EEncodingOptions::ForceUTF8WithoutBOM))
 	{
 		UE_LOG(LogSourceControl, Error, TEXT("Failed to write to temp file: %s"), *Filename);
 	}
@@ -53,10 +53,9 @@ static bool RunCommandInternalRaw(const FString& InCommand, const FString& InPat
 	FString FullCommand;
 	FString LogableCommand; // short version of the command for logging purpose
 
-	if (!InRepositoryRoot.IsEmpty())
+	if(!InRepositoryRoot.IsEmpty())
 	{
 		// Specify the working copy (the root) of the git repository (before the command itself)
-		// @todo Does not work in UE4.1 on Mac if there is a space in the path ("[REDACTED_USER_HOME]/Unreal Project/MyProject")
 		FullCommand  = TEXT("--work-tree=\"");
 		FullCommand += InRepositoryRoot;
 		// and the ".git" subdirectory in it (before the command itself)
@@ -112,32 +111,30 @@ static bool RunCommandInternal(const FString& InCommand, const FString& InPathTo
 
 FString FindGitBinaryPath()
 {
-	bool bFound = false;
-
 #if PLATFORM_WINDOWS
-	// 1) First of all, check for the ThirdParty directory as it may contain a specific version of Git for this plugin to work
+	// 1) First of all, look into standard install directory
 	// NOTE using only "git" (or "git.exe") relying on the "PATH" envvar does not always work as expected, depending on the installation:
 	// If the PATH is set with "git/cmd" instead of "git/bin",
 	// "git.exe" launch "git/cmd/git.exe" that redirect to "git/bin/git.exe" and ExecProcess() is unable to catch its outputs streams.
-
-	// Under Windows, we can use the third party "msysgit PortableGit" https://code.google.com/p/msysgit/downloads/list?can=1&q=PortableGit
-	// NOTE: Win32 platform subdirectory as there is no Git 64bit build available
-	FString GitBinaryPath(FPaths::EngineDir() / TEXT("Binaries/ThirdParty/git/Win32/bin") / TEXT("git.exe"));
-	bFound = CheckGitAvailability(GitBinaryPath);
-#else
-	FString GitBinaryPath;
-#endif
-
-	// 2) If Git is not found in ThirdParty directory, look into standard install directory
+	FString GitBinaryPath(TEXT("C:/Program Files (x86)/Git/bin/git.exe"));
+	bool bFound = CheckGitAvailability(GitBinaryPath);
+	
+	// 2) If Git is not found in default install, look for the PortableGit provided by GitHub for Windows
 	if(!bFound)
 	{
-#if PLATFORM_WINDOWS
-		// @todo use the Windows registry to find Git
-		GitBinaryPath = TEXT("C:/Program Files (x86)/Git/bin/git.exe");
+		// The latest GitHub for windows adds its binaries into the local appdata directory:
+		// [REDACTED_USER_HOME]\AppData\Local\GitHub\PortableGit_c2ba306e536fdf878271f7fe636a147ff37326ad\bin
+		TCHAR AppDataLocalPath[4096];
+		FPlatformMisc::GetEnvironmentVariable(TEXT("LOCALAPPDATA"), AppDataLocalPath, ARRAY_COUNT(AppDataLocalPath));
+		FString SearchPath = FString::Printf(TEXT("%s/GitHub/PortableGit_*"), AppDataLocalPath);
+		TArray<FString> PortableGitFolders;
+		IFileManager::Get().FindFiles(PortableGitFolders, *SearchPath, false, true);
+		// FindFiles just returns directory names, so we need to prepend the root path to get the full path.
+		GitBinaryPath = FString::Printf(TEXT("%s/GitHub/%s/bin/git.exe"), AppDataLocalPath, *(PortableGitFolders.Last())); // keep only the last PortableGit found
+	}
 #else
-		GitBinaryPath = TEXT("/usr/bin/git");
+	FString GitBinaryPath = TEXT("/usr/bin/git");
 #endif
-	}
 
 	FPaths::MakePlatformFilename( GitBinaryPath );
 	return GitBinaryPath;
diff --git a/Engine/Plugins/Developer/GitSourceControl/Source/GitSourceControl/Private/GitSourceControlUtils.h b/Engine/Plugins/Developer/GitSourceControl/Source/GitSourceControl/Private/GitSourceControlUtils.h
--- a/Engine/Plugins/Developer/GitSourceControl/Source/GitSourceControl/Private/GitSourceControlUtils.h
+++ b/Engine/Plugins/Developer/GitSourceControl/Source/GitSourceControl/Private/GitSourceControlUtils.h
@@ -6,6 +6,8 @@
 #include "GitSourceControlRevision.h"
 #include "GitSourceControlCommand.h"
 
+class FGitSourceControlCommand;
+
 /**
  * Helper struct for maintaining temporary files for passing to commands
  */
