PR #958: Fix git commit batching (commit could have failed for large number of files) (Contributed by [REDACTED_PERSON])

diff --git a/Engine/Plugins/Developer/GitSourceControl/Source/GitSourceControl/Private/GitSourceControlUtils.cpp b/Engine/Plugins/Developer/GitSourceControl/Source/GitSourceControl/Private/GitSourceControlUtils.cpp
--- a/Engine/Plugins/Developer/GitSourceControl/Source/GitSourceControl/Private/GitSourceControlUtils.cpp
+++ b/Engine/Plugins/Developer/GitSourceControl/Source/GitSourceControl/Private/GitSourceControlUtils.cpp
@@ -232,15 +232,13 @@ bool RunCommit(const FString& InPathToGitBinary, const FString& InRepositoryRoot
 	{
 		// Batch files up so we dont exceed command-line limits
 		int32 FileCount = 0;
+		TArray<FString> FilesInBatch;
+		for(int32 FileIndex = 0; FileIndex < GitSourceControlConstants::MaxFilesPerBatch; FileIndex++, FileCount++)
 		{
-			TArray<FString> FilesInBatch;
-			for(int32 FileIndex = 0; FileIndex < GitSourceControlConstants::MaxFilesPerBatch; FileIndex++, FileCount++)
-			{
-				FilesInBatch.Add(InFiles[FileCount]);
-			}
+			FilesInBatch.Add(InFiles[FileCount]);
 		}
 		// First batch is a simple "git commit" command with only the first files
-		bResult &= RunCommandInternal(TEXT("commit"), InPathToGitBinary, InRepositoryRoot, InParameters, InFiles, OutResults, OutErrorMessages);
+		bResult &= RunCommandInternal(TEXT("commit"), InPathToGitBinary, InRepositoryRoot, InParameters, FilesInBatch, OutResults, OutErrorMessages);
 		
 		TArray<FString> Parameters;
 		for(const auto& Parameter : InParameters)
