Fix for crash if you don't have Git.

diff --git a/Engine/Plugins/Developer/GitSourceControl/Source/GitSourceControl/Private/GitSourceControlUtils.cpp b/Engine/Plugins/Developer/GitSourceControl/Source/GitSourceControl/Private/GitSourceControlUtils.cpp
--- a/Engine/Plugins/Developer/GitSourceControl/Source/GitSourceControl/Private/GitSourceControlUtils.cpp
+++ b/Engine/Plugins/Developer/GitSourceControl/Source/GitSourceControl/Private/GitSourceControlUtils.cpp
@@ -129,8 +129,13 @@ FString FindGitBinaryPath()
 		FString SearchPath = FString::Printf(TEXT("%s/GitHub/PortableGit_*"), AppDataLocalPath);
 		TArray<FString> PortableGitFolders;
 		IFileManager::Get().FindFiles(PortableGitFolders, *SearchPath, false, true);
-		// FindFiles just returns directory names, so we need to prepend the root path to get the full path.
-		GitBinaryPath = FString::Printf(TEXT("%s/GitHub/%s/bin/git.exe"), AppDataLocalPath, *(PortableGitFolders.Last())); // keep only the last PortableGit found
+		
+		// If we found a Portable Git setup the path to it, otherwise we leave GitBinaryPath = to the default set above.
+		if (PortableGitFolders.Num() > 0)
+		{
+			// FindFiles just returns directory names, so we need to prepend the root path to get the full path.
+			GitBinaryPath = FString::Printf(TEXT("%s/GitHub/%s/bin/git.exe"), AppDataLocalPath, *(PortableGitFolders.Last())); // keep only the last PortableGit found
+		}
 	}
 #else
 	FString GitBinaryPath = TEXT("/usr/bin/git");
