PR #966: (Git Plugin) shearch more paths for git executable if not found in standard path (Contributed by [REDACTED_PERSON])

diff --git a/Engine/Plugins/Developer/GitSourceControl/Source/GitSourceControl/Private/GitSourceControlUtils.cpp b/Engine/Plugins/Developer/GitSourceControl/Source/GitSourceControl/Private/GitSourceControlUtils.cpp
--- a/Engine/Plugins/Developer/GitSourceControl/Source/GitSourceControl/Private/GitSourceControlUtils.cpp
+++ b/Engine/Plugins/Developer/GitSourceControl/Source/GitSourceControl/Private/GitSourceControlUtils.cpp
@@ -118,8 +118,25 @@ FString FindGitBinaryPath()
 	// "git.exe" launch "git/cmd/git.exe" that redirect to "git/bin/git.exe" and ExecProcess() is unable to catch its outputs streams.
 	FString GitBinaryPath(TEXT("C:/Program Files (x86)/Git/bin/git.exe"));
 	bool bFound = CheckGitAvailability(GitBinaryPath);
-	
-	// 2) If Git is not found in default install, look for the PortableGit provided by GitHub for Windows
+
+	// 2) Else, look for the version of Git bundled with SmartGit "Installer with JRE"
+	if(!bFound)
+	{
+		FString GitBinaryPath(TEXT("C:/Program Files (x86)/SmartGit/bin/git.exe"));
+		bFound = CheckGitAvailability(GitBinaryPath);
+	}
+
+	// 3) Else, look for the local_git provided by SourceTree
+	if(!bFound)
+	{
+		// [REDACTED_USER_HOME]\AppData\Local\Atlassian\SourceTree\git_local\bin
+		TCHAR AppDataLocalPath[4096];
+		FPlatformMisc::GetEnvironmentVariable(TEXT("LOCALAPPDATA"), AppDataLocalPath, ARRAY_COUNT(AppDataLocalPath));
+		GitBinaryPath = FString::Printf(TEXT("%s/Atlassian/SourceTree/git_local/bin/git.exe"), AppDataLocalPath);
+		bFound = CheckGitAvailability(GitBinaryPath);
+	}
+
+	// 4) Else, look for the PortableGit provided by GitHub for Windows
 	if(!bFound)
 	{
 		// The latest GitHub for windows adds its binaries into the local appdata directory:
@@ -129,19 +146,28 @@ FString FindGitBinaryPath()
 		FString SearchPath = FString::Printf(TEXT("%s/GitHub/PortableGit_*"), AppDataLocalPath);
 		TArray<FString> PortableGitFolders;
 		IFileManager::Get().FindFiles(PortableGitFolders, *SearchPath, false, true);
-		
-		// If we found a Portable Git setup the path to it, otherwise we leave GitBinaryPath = to the default set above.
-		if (PortableGitFolders.Num() > 0)
+		if(PortableGitFolders.Num() > 0)
 		{
 			// FindFiles just returns directory names, so we need to prepend the root path to get the full path.
 			GitBinaryPath = FString::Printf(TEXT("%s/GitHub/%s/bin/git.exe"), AppDataLocalPath, *(PortableGitFolders.Last())); // keep only the last PortableGit found
+			bFound = CheckGitAvailability(GitBinaryPath);
 		}
 	}
 #else
 	FString GitBinaryPath = TEXT("/usr/bin/git");
+	bool bFound = CheckGitAvailability(GitBinaryPath);
 #endif
 
-	FPaths::MakePlatformFilename( GitBinaryPath );
+	if(bFound)
+	{
+		FPaths::MakePlatformFilename(GitBinaryPath);
+	}
+	else
+	{
+		// If we did not find a path to Git, set it empty
+		GitBinaryPath.Empty();
+	}
+
 	return GitBinaryPath;
 }
 
@@ -154,17 +180,20 @@ bool CheckGitAvailability(const FString& InPathToGitBinary)
 	bGitAvailable = RunCommandInternalRaw(TEXT("version"), InPathToGitBinary, FString(), TArray<FString>(), TArray<FString>(), InfoMessages, ErrorMessages);
 	if(bGitAvailable)
 	{
-		if(InfoMessages.Contains("git"))
-		{
-			bGitAvailable = true;
-		}
-		else
+		if(!InfoMessages.Contains("git"))
 		{
 			bGitAvailable = false;
 		}
 	}
 
-	// @todo also check Git config user.name & user.email
+	if(bGitAvailable)
+	{
+		UE_LOG(LogSourceControl, Log, TEXT("'%s' found"), *InPathToGitBinary);
+	}
+	else
+	{
+		UE_LOG(LogSourceControl, Log, TEXT("'%s' not found"), *InPathToGitBinary);
+	}
 
 	return bGitAvailable;
 }
@@ -233,7 +262,7 @@ bool RunCommit(const FString& InPathToGitBinary, const FString& InRepositoryRoot
 {
 	bool bResult = true;
 
-	if (InFiles.Num() > GitSourceControlConstants::MaxFilesPerBatch)
+	if(InFiles.Num() > GitSourceControlConstants::MaxFilesPerBatch)
 	{
 		// Batch files up so we dont exceed command-line limits
 		int32 FileCount = 0;
@@ -290,7 +319,7 @@ public:
 		FString RelativeFilename = InResult.RightChop(3);
 		// Note: this is not enough in case of a rename from -> to
 		int32 RenameIndex;
-		if (RelativeFilename.FindLastChar('>', RenameIndex))
+		if(RelativeFilename.FindLastChar('>', RenameIndex))
 		{
 			// Extract only the second part of a rename "from -> to"
 			RelativeFilename = RelativeFilename.RightChop(RenameIndex + 2);
diff --git a/Engine/Plugins/Developer/GitSourceControl/Source/GitSourceControl/Private/GitSourceControlUtils.h b/Engine/Plugins/Developer/GitSourceControl/Source/GitSourceControl/Private/GitSourceControlUtils.h
--- a/Engine/Plugins/Developer/GitSourceControl/Source/GitSourceControl/Private/GitSourceControlUtils.h
+++ b/Engine/Plugins/Developer/GitSourceControl/Source/GitSourceControl/Private/GitSourceControlUtils.h
@@ -4,7 +4,6 @@
 
 #include "GitSourceControlState.h"
 #include "GitSourceControlRevision.h"
-#include "GitSourceControlCommand.h"
 
 class FGitSourceControlCommand;
 
@@ -33,8 +32,8 @@ namespace GitSourceControlUtils
 {
 
 /**
- * Find the path to the Git binary, looking in a few standard place (ThirdParty subdirectory and standard sSystem paths)
- * @returns the path to the Git binary if found, or the last path tested if no git found.
+ * Find the path to the Git binary, looking into a few places (standalone Git install, and other common tools embedding Git)
+ * @returns the path to the Git binary if found, or an empty string.
  */
 FString FindGitBinaryPath();
 
