PR #1029: (Git Plugin) improved status with user name and e-mail (Contributed by [REDACTED_PERSON])

diff --git a/Engine/Plugins/Developer/GitSourceControl/Source/GitSourceControl/Private/GitSourceControlCommand.cpp b/Engine/Plugins/Developer/GitSourceControl/Source/GitSourceControl/Private/GitSourceControlCommand.cpp
--- a/Engine/Plugins/Developer/GitSourceControl/Source/GitSourceControl/Private/GitSourceControlCommand.cpp
+++ b/Engine/Plugins/Developer/GitSourceControl/Source/GitSourceControl/Private/GitSourceControlCommand.cpp
@@ -21,7 +21,6 @@ FGitSourceControlCommand::FGitSourceControlCommand(const TSharedRef<class ISourc
 	FGitSourceControlModule& GitSourceControl = FModuleManager::LoadModuleChecked<FGitSourceControlModule>( "GitSourceControl" );
 	PathToGitBinary = GitSourceControl.AccessSettings().GetBinaryPath();
 	PathToRepositoryRoot = GitSourceControl.GetProvider().GetPathToRepositoryRoot();
-	PathToGameDir = GitSourceControl.GetProvider().GetPathToGameDir();
 }
 
 bool FGitSourceControlCommand::DoWork()
diff --git a/Engine/Plugins/Developer/GitSourceControl/Source/GitSourceControl/Private/GitSourceControlCommand.h b/Engine/Plugins/Developer/GitSourceControl/Source/GitSourceControl/Private/GitSourceControlCommand.h
--- a/Engine/Plugins/Developer/GitSourceControl/Source/GitSourceControl/Private/GitSourceControlCommand.h
+++ b/Engine/Plugins/Developer/GitSourceControl/Source/GitSourceControl/Private/GitSourceControlCommand.h
@@ -38,12 +38,6 @@ public:
 	/** Path to the root of the Git repository: can be the GameDir itself, or any parent directory (found by the "Connect" operation) */
 	FString PathToRepositoryRoot;
 
-	/** Path to the Game directory: shall be inside of the Git repository */
-	FString PathToGameDir;
-
-	/** Name of the current branch (found by the "Connect" operation) */
-	FString BranchName;
-
 	/** Operation we want to perform - contains outward-facing parameters & results */
 	TSharedRef<class ISourceControlOperation, ESPMode::ThreadSafe> Operation;
 
diff --git a/Engine/Plugins/Developer/GitSourceControl/Source/GitSourceControl/Private/GitSourceControlModule.cpp b/Engine/Plugins/Developer/GitSourceControl/Source/GitSourceControl/Private/GitSourceControlModule.cpp
--- a/Engine/Plugins/Developer/GitSourceControl/Source/GitSourceControl/Private/GitSourceControlModule.cpp
+++ b/Engine/Plugins/Developer/GitSourceControl/Source/GitSourceControl/Private/GitSourceControlModule.cpp
@@ -47,11 +47,6 @@ void FGitSourceControlModule::ShutdownModule()
 	IModularFeatures::Get().UnregisterModularFeature("SourceControl", &GitSourceControlProvider);
 }
 
-FGitSourceControlSettings& FGitSourceControlModule::AccessSettings()
-{
-	return GitSourceControlSettings;
-}
-
 void FGitSourceControlModule::SaveSettings()
 {
 	if (FApp::IsUnattended() || IsRunningCommandlet())
diff --git a/Engine/Plugins/Developer/GitSourceControl/Source/GitSourceControl/Private/GitSourceControlModule.h b/Engine/Plugins/Developer/GitSourceControl/Source/GitSourceControl/Private/GitSourceControlModule.h
--- a/Engine/Plugins/Developer/GitSourceControl/Source/GitSourceControl/Private/GitSourceControlModule.h
+++ b/Engine/Plugins/Developer/GitSourceControl/Source/GitSourceControl/Private/GitSourceControlModule.h
@@ -60,7 +60,10 @@ public:
 	virtual void ShutdownModule() override;
 
 	/** Access the Git source control settings */
-	FGitSourceControlSettings& AccessSettings();
+	FGitSourceControlSettings& AccessSettings()
+	{
+		return GitSourceControlSettings;
+	}
 
 	/** Save the Git source control settings */
 	void SaveSettings();
diff --git a/Engine/Plugins/Developer/GitSourceControl/Source/GitSourceControl/Private/GitSourceControlOperations.cpp b/Engine/Plugins/Developer/GitSourceControl/Source/GitSourceControl/Private/GitSourceControlOperations.cpp
--- a/Engine/Plugins/Developer/GitSourceControl/Source/GitSourceControl/Private/GitSourceControlOperations.cpp
+++ b/Engine/Plugins/Developer/GitSourceControl/Source/GitSourceControl/Private/GitSourceControlOperations.cpp
@@ -18,31 +18,16 @@ bool FGitConnectWorker::Execute(FGitSourceControlCommand& InCommand)
 {
 	check(InCommand.Operation->GetName() == GetName());
 
-	InCommand.bCommandSuccessful = GitSourceControlUtils::FindRootDirectory(InCommand.PathToGameDir, InCommand.PathToRepositoryRoot);
-	if(InCommand.bCommandSuccessful)
-	{
-		InCommand.bCommandSuccessful = GitSourceControlUtils::RunCommand(TEXT("status"), InCommand.PathToGitBinary, InCommand.PathToRepositoryRoot, TArray<FString>(), TArray<FString>(), InCommand.InfoMessages, InCommand.ErrorMessages);
-	}
+	InCommand.bCommandSuccessful = GitSourceControlUtils::RunCommand(TEXT("status"), InCommand.PathToGitBinary, InCommand.PathToRepositoryRoot, TArray<FString>(), TArray<FString>(), InCommand.InfoMessages, InCommand.ErrorMessages);
 	if(!InCommand.bCommandSuccessful || InCommand.ErrorMessages.Num() > 0 || InCommand.InfoMessages.Num() == 0)
 	{
-		// @todo popup to propose to initialize the git repository "git init + .gitignore"
+		// @todo popup to propose to initialize the git repository "git init + .gitignore" (easier in the provider!)
+		// then "Connect" could in fact do the "git init"
 		StaticCastSharedRef<FConnect>(InCommand.Operation)->SetErrorText(LOCTEXT("NotAWorkingCopyError", "Project is not part of a Git working copy."));
 		// @todo Double error messages (and displayed in reverse order): Perforce distinguish the two errors
 		InCommand.ErrorMessages.Add(LOCTEXT("NotAWorkingCopyErrorHelp", "You should check out a working copy into your project directory.").ToString());
 		InCommand.bCommandSuccessful = false;
 	}
-	else // if(InCommand.bCommandSuccessful)
-	{
-		TArray<FString> Parameters;
-		Parameters.Add("--short HEAD");
-
-		// Get current branche name
-		GitSourceControlUtils::RunCommand(TEXT("symbolic-ref"), InCommand.PathToGitBinary, InCommand.PathToRepositoryRoot, Parameters, TArray<FString>(), InCommand.InfoMessages, InCommand.ErrorMessages);
-		if(InCommand.InfoMessages.Num() == 1)
-		{
-			InCommand.BranchName = InCommand.InfoMessages[0];
-		}
-	}
 
 	return InCommand.bCommandSuccessful;
 }
diff --git a/Engine/Plugins/Developer/GitSourceControl/Source/GitSourceControl/Private/GitSourceControlProvider.cpp b/Engine/Plugins/Developer/GitSourceControl/Source/GitSourceControl/Private/GitSourceControlProvider.cpp
--- a/Engine/Plugins/Developer/GitSourceControl/Source/GitSourceControl/Private/GitSourceControlProvider.cpp
+++ b/Engine/Plugins/Developer/GitSourceControl/Source/GitSourceControl/Private/GitSourceControlProvider.cpp
@@ -18,15 +18,9 @@ static FName ProviderName("Git");
 
 void FGitSourceControlProvider::Init(bool bForceConnection)
 {
-	PathToGameDir = FPaths::ConvertRelativePathToFull(FPaths::GameDir());
-
 	CheckGitAvailability();
 
-	if(bForceConnection)
-	{
-		// First connection, to find the path to the root git directory (if any)
-		GitSourceControlUtils::FindRootDirectory(PathToGameDir, PathToRepositoryRoot);
-	}
+	// bForceConnection: not used anymore
 }
 
 void FGitSourceControlProvider::CheckGitAvailability()
@@ -36,10 +30,29 @@ void FGitSourceControlProvider::CheckGitAvailability()
 	if(!PathToGitBinary.IsEmpty())
 	{
 		bGitAvailable = GitSourceControlUtils::CheckGitAvailability(PathToGitBinary);
+		if(bGitAvailable)
+		{
+			// Find the path to the root Git directory (if any)
+			const FString PathToGameDir = FPaths::ConvertRelativePathToFull(FPaths::GameDir());
+			const bool bRepositoryFound = GitSourceControlUtils::FindRootDirectory(PathToGameDir, PathToRepositoryRoot);
+			if(bRepositoryFound)
+			{
+				// Get user name & email
+				GitSourceControlUtils::GetUserConfig(PathToGitBinary, PathToRepositoryRoot, UserName, UserEmail);
+				// Get branch name
+				GitSourceControlUtils::GetBranchName(PathToGitBinary, PathToRepositoryRoot, BranchName);
+			}
+			else
+			{
+				UE_LOG(LogSourceControl, Error, TEXT("'%s' is not part of a Git repository"), *FPaths::GameDir());
+				bGitAvailable = false;
+			}
+		}
 	}
 	else
 	{
-		UE_LOG(LogSourceControl, Error, TEXT("Git not found"), *PathToGitBinary);
+		UE_LOG(LogSourceControl, Error, TEXT("Git not found"));
+		bGitAvailable = false;
 	}
 }
 
@@ -69,11 +82,12 @@ TSharedRef<FGitSourceControlState, ESPMode::ThreadSafe> FGitSourceControlProvide
 FText FGitSourceControlProvider::GetStatusText() const
 {
 	FFormatNamedArguments Args;
-	Args.Add( TEXT("IsEnabled"), IsEnabled() ? LOCTEXT("Yes", "Yes") : LOCTEXT("No", "No") );
 	Args.Add( TEXT("RepositoryName"), FText::FromString(PathToRepositoryRoot) );
 	Args.Add( TEXT("BranchName"), FText::FromString(BranchName) );
+	Args.Add( TEXT("UserName"), FText::FromString(UserName) );
+	Args.Add( TEXT("UserEmail"), FText::FromString(UserEmail) );
 
-	return FText::Format( NSLOCTEXT("Status", "Provider: Git\nEnabledLabel", "Enabled: {IsEnabled}\nRepository: {RepositoryName}\nBranch: {BranchName}"), Args );
+	return FText::Format( NSLOCTEXT("Status", "Provider: Git\nEnabledLabel", "Repository: {RepositoryName}\nBranch: {BranchName}\nUser: {UserName}\nE-mail: {UserEmail}"), Args );
 }
 
 bool FGitSourceControlProvider::IsEnabled() const
@@ -116,10 +130,10 @@ ECommandResult::Type FGitSourceControlProvider::GetState( const TArray<FString>&
 TArray<FSourceControlStateRef> FGitSourceControlProvider::GetCachedStateByPredicate(const TFunctionRef<bool(const FSourceControlStateRef&)>& Predicate) const
 {
 	TArray<FSourceControlStateRef> Result;
-	for (const auto& CacheItem : StateCache)
+	for(const auto& CacheItem : StateCache)
 	{
 		FSourceControlStateRef State = CacheItem.Value;
-		if (Predicate(State))
+		if(Predicate(State))
 		{
 			Result.Add(State);
 		}
@@ -176,15 +190,11 @@ ECommandResult::Type FGitSourceControlProvider::Execute( const TSharedRef<ISourc
 	if(InConcurrency == EConcurrency::Synchronous)
 	{
 		Command->bAutoDelete = false;
-// @todo: temporary debug logs
-//		UE_LOG(LogSourceControl, Log, TEXT("ExecuteSynchronousCommand: '%s'"), *InOperation->GetName().ToString());
 		return ExecuteSynchronousCommand(*Command, InOperation->GetInProgressString());
 	}
 	else
 	{
 		Command->bAutoDelete = true;
-// @todo: temporary debug logs
-//		UE_LOG(LogSourceControl, Log, TEXT("IssueCommand(Asynchronous): '%s'"), *InOperation->GetName().ToString());
 		return IssueCommand(*Command);
 	}
 }
@@ -211,12 +221,12 @@ bool FGitSourceControlProvider::UsesChangelists() const
 TSharedPtr<IGitSourceControlWorker, ESPMode::ThreadSafe> FGitSourceControlProvider::CreateWorker(const FName& InOperationName) const
 {
 	const FGetGitSourceControlWorker* Operation = WorkersMap.Find(InOperationName);
-	if(Operation != NULL)
+	if(Operation != nullptr)
 	{
 		return Operation->Execute();
 	}
 		
-	return NULL;
+	return nullptr;
 }
 
 void FGitSourceControlProvider::RegisterWorker( const FName& InName, const FGetGitSourceControlWorker& InDelegate )
@@ -245,7 +255,7 @@ void FGitSourceControlProvider::Tick()
 	for(int32 CommandIndex = 0; CommandIndex < CommandQueue.Num(); ++CommandIndex)
 	{
 		FGitSourceControlCommand& Command = *CommandQueue[CommandIndex];
-		if (Command.bExecuteProcessed)
+		if(Command.bExecuteProcessed)
 		{
 			// Remove command from the queue
 			CommandQueue.RemoveAt(CommandIndex);
@@ -253,13 +263,6 @@ void FGitSourceControlProvider::Tick()
 			// let command update the states of any files
 			bStatesUpdated |= Command.Worker->UpdateStates();
 
-			// if connect operation: obtain the path to the Git root directory of the repository
-			if(Command.Operation->GetName() == "Connect" && Command.bCommandSuccessful)
-			{
-				PathToRepositoryRoot = Command.PathToRepositoryRoot;
-				BranchName = Command.BranchName;
-			}
-
 			// dump any messages to output log
 			OutputCommandMessages(Command);
 
@@ -341,7 +344,7 @@ ECommandResult::Type FGitSourceControlProvider::ExecuteSynchronousCommand(FGitSo
 
 ECommandResult::Type FGitSourceControlProvider::IssueCommand(FGitSourceControlCommand& InCommand)
 {
-	if ( GThreadPool != NULL )
+	if(GThreadPool != nullptr)
 	{
 		// Queue this to our worker thread(s) for resolving
 		GThreadPool->AddQueuedWork(&InCommand);
diff --git a/Engine/Plugins/Developer/GitSourceControl/Source/GitSourceControl/Private/GitSourceControlProvider.h b/Engine/Plugins/Developer/GitSourceControl/Source/GitSourceControl/Private/GitSourceControlProvider.h
--- a/Engine/Plugins/Developer/GitSourceControl/Source/GitSourceControl/Private/GitSourceControlProvider.h
+++ b/Engine/Plugins/Developer/GitSourceControl/Source/GitSourceControl/Private/GitSourceControlProvider.h
@@ -59,10 +59,16 @@ public:
 		return PathToRepositoryRoot;
 	}
 
-	/** Get the path to the Game directory: shall be inside of the Git repository */
-	inline const FString& GetPathToGameDir() const
+	/** Git config user.name */
+	inline const FString& GetUserName() const
 	{
-		return PathToGameDir;
+		return UserName;
+	}
+
+	/** Git config user.email */
+	inline const FString& GetUserEmail() const
+	{
+		return UserEmail;
 	}
 
 	/** Helper function used to update state cache */
@@ -93,10 +99,13 @@ private:
 	/** Path to the root of the Git repository: can be the GameDir itself, or any parent directory (found by the "Connect" operation) */
 	FString PathToRepositoryRoot;
 
-	/** Path to the Game directory: shall be inside of the Git repository */
-	FString PathToGameDir;
+	/** Git config user.name (from local repository, else globally) */
+	FString UserName;
+
+	/** Git config user.email (from local repository, else globally) */
+	FString UserEmail;
 
-	/** Name of the current branch (found by the "Connect" operation) */
+	/** Name of the current branch */
 	FString BranchName;
 
 	/** State cache */
diff --git a/Engine/Plugins/Developer/GitSourceControl/Source/GitSourceControl/Private/GitSourceControlSettings.cpp b/Engine/Plugins/Developer/GitSourceControl/Source/GitSourceControl/Private/GitSourceControlSettings.cpp
--- a/Engine/Plugins/Developer/GitSourceControl/Source/GitSourceControl/Private/GitSourceControlSettings.cpp
+++ b/Engine/Plugins/Developer/GitSourceControl/Source/GitSourceControl/Private/GitSourceControlSettings.cpp
@@ -27,6 +27,7 @@ void FGitSourceControlSettings::SetBinaryPath(const FString& InString)
 	BinaryPath = InString;
 }
 
+// This is called at startup nearly before anything else in our module: BinaryPath will then be used by the provider
 void FGitSourceControlSettings::LoadSettings()
 {
 	FScopeLock ScopeLock(&CriticalSection);
@@ -41,9 +42,13 @@ void FGitSourceControlSettings::LoadSettings()
 void FGitSourceControlSettings::SaveSettings() const
 {
 	FScopeLock ScopeLock(&CriticalSection);
-	const FString& IniFile = SourceControlHelpers::GetSettingsIni();
-	GConfig->SetString(*GitSettingsConstants::SettingsSection, TEXT("BinaryPath"), *BinaryPath, IniFile);
 
+	// Re-Check provided git binary path for each change
 	FGitSourceControlModule& GitSourceControl = FModuleManager::LoadModuleChecked<FGitSourceControlModule>("GitSourceControl");
 	GitSourceControl.GetProvider().CheckGitAvailability();
+	if (GitSourceControl.GetProvider().IsAvailable())
+	{
+		const FString& IniFile = SourceControlHelpers::GetSettingsIni();
+		GConfig->SetString(*GitSettingsConstants::SettingsSection, TEXT("BinaryPath"), *BinaryPath, IniFile);
+	}
 }
\ No newline at end of file
diff --git a/Engine/Plugins/Developer/GitSourceControl/Source/GitSourceControl/Private/GitSourceControlUtils.cpp b/Engine/Plugins/Developer/GitSourceControl/Source/GitSourceControl/Private/GitSourceControlUtils.cpp
--- a/Engine/Plugins/Developer/GitSourceControl/Source/GitSourceControl/Private/GitSourceControlUtils.cpp
+++ b/Engine/Plugins/Developer/GitSourceControl/Source/GitSourceControl/Private/GitSourceControlUtils.cpp
@@ -82,14 +82,17 @@ static bool RunCommandInternalRaw(const FString& InCommand, const FString& InPat
 
 	FullCommand += LogableCommand;
 
-	UE_LOG(LogSourceControl, Log, TEXT("RunCommandInternalRaw: 'git %s'"), *LogableCommand);
 // @todo: temporary debug logs
 //	UE_LOG(LogSourceControl, Log, TEXT("RunCommandInternalRaw: 'git %s'"), *FullCommand);
 	FPlatformProcess::ExecProcess(*InPathToGitBinary, *FullCommand, &ReturnCode, &OutResults, &OutErrors);
-	UE_LOG(LogSourceControl, Log, TEXT("RunCommandInternalRaw: ExecProcess ReturnCode=%d OutResults='%s'"), ReturnCode, *OutResults);
-	if(!OutErrors.IsEmpty())
+	if(OutErrors.IsEmpty() && (ReturnCode == 0))
 	{
-		UE_LOG(LogSourceControl, Error, TEXT("RunCommandInternalRaw: ExecProcess ReturnCode=%d OutErrors='%s'"), ReturnCode, *OutErrors);
+		UE_LOG(LogSourceControl, Log, TEXT("'git %s':\n%s"), *LogableCommand, *OutResults);
+	}
+	else
+	{
+		// @todo filter here like RemoveRedundantErrors("' is outside repository") ?
+		UE_LOG(LogSourceControl, Error, TEXT("'git %s': ReturnCode=%d Errors='%s'"), *LogableCommand, ReturnCode, *OutErrors);
 	}
 
 	return ReturnCode == 0;
@@ -198,6 +201,7 @@ bool CheckGitAvailability(const FString& InPathToGitBinary)
 	return bGitAvailable;
 }
 
+// Find the root of the Git repository, looking from the GameDir and upward in its parent directories.
 bool FindRootDirectory(const FString& InPathToGameDir, FString& OutRepositoryRoot)
 {
 	bool bFound = false;
@@ -207,7 +211,7 @@ bool FindRootDirectory(const FString& InPathToGameDir, FString& OutRepositoryRoo
 	while(!bFound && !OutRepositoryRoot.IsEmpty())
 	{
 		PathToGitSubdirectory = OutRepositoryRoot;
-		PathToGitSubdirectory += TEXT(".git"); // Look for the ".git" subdirectory at the root of every Git repository
+		PathToGitSubdirectory += TEXT(".git"); // Look for the ".git" subdirectory present at the root of every Git repository
 		bFound = IFileManager::Get().DirectoryExists(*PathToGitSubdirectory);
 		if(!bFound)
 		{
@@ -227,6 +231,57 @@ bool FindRootDirectory(const FString& InPathToGameDir, FString& OutRepositoryRoo
 	return bFound;
 }
 
+void GetUserConfig(const FString& InPathToGitBinary, const FString& InRepositoryRoot, FString& OutUserName, FString& OutUserEmail)
+{
+	bool bResults;
+	TArray<FString> InfoMessages;
+	TArray<FString> ErrorMessages;
+	TArray<FString> Parameters;
+	Parameters.Add(TEXT("user.name"));
+	bResults = RunCommandInternal(TEXT("config"), InPathToGitBinary, InRepositoryRoot, Parameters, TArray<FString>(), InfoMessages, ErrorMessages);
+	if (bResults && InfoMessages.Num() > 0)
+	{
+		OutUserName = InfoMessages[0];
+	}
+
+	Parameters.Reset();
+	Parameters.Add(TEXT("user.email"));
+	InfoMessages.Reset();
+	bResults &= RunCommandInternal(TEXT("config"), InPathToGitBinary, InRepositoryRoot, Parameters, TArray<FString>(), InfoMessages, ErrorMessages);
+	if (bResults && InfoMessages.Num() > 0)
+	{
+		OutUserEmail = InfoMessages[0];
+	}
+}
+
+void GetBranchName(const FString& InPathToGitBinary, const FString& InRepositoryRoot, FString& OutBranchName)
+{
+	bool bResults;
+	TArray<FString> InfoMessages;
+	TArray<FString> ErrorMessages;
+	TArray<FString> Parameters;
+	Parameters.Add(TEXT("--short"));
+	Parameters.Add(TEXT("--quiet"));		// no error message while in detached HEAD
+	Parameters.Add(TEXT("HEAD"));	
+	bResults = RunCommandInternal(TEXT("symbolic-ref"), InPathToGitBinary, InRepositoryRoot, Parameters, TArray<FString>(), InfoMessages, ErrorMessages);
+	if (bResults && InfoMessages.Num() > 0)
+	{
+		OutBranchName = InfoMessages[0];
+	}
+	else
+	{
+		Parameters.Reset();
+		Parameters.Add(TEXT("-1"));
+		Parameters.Add(TEXT("--format=\"%h\""));		// no error message while in detached HEAD
+		bResults = RunCommandInternal(TEXT("log"), InPathToGitBinary, InRepositoryRoot, Parameters, TArray<FString>(), InfoMessages, ErrorMessages);
+		if (bResults && InfoMessages.Num() > 0)
+		{
+			OutBranchName = "HEAD detached at ";
+			OutBranchName += InfoMessages[0];
+		}
+	}
+}
+
 bool RunCommand(const FString& InCommand, const FString& InPathToGitBinary, const FString& InRepositoryRoot, const TArray<FString>& InParameters, const TArray<FString>& InFiles, TArray<FString>& OutResults, TArray<FString>& OutErrorMessages)
 {
 	bool bResult = true;
diff --git a/Engine/Plugins/Developer/GitSourceControl/Source/GitSourceControl/Private/GitSourceControlUtils.h b/Engine/Plugins/Developer/GitSourceControl/Source/GitSourceControl/Private/GitSourceControlUtils.h
--- a/Engine/Plugins/Developer/GitSourceControl/Source/GitSourceControl/Private/GitSourceControlUtils.h
+++ b/Engine/Plugins/Developer/GitSourceControl/Source/GitSourceControl/Private/GitSourceControlUtils.h
@@ -49,9 +49,26 @@ bool CheckGitAvailability(const FString& InPathToGitBinary);
  * @param InPathToGameDir		The path to the Game Directory
  * @param OutRepositoryRoot		The path to the root directory of the Git repository if found
  * @returns true if the command succeeded and returned no errors
-*/
+ */
 bool FindRootDirectory(const FString& InPathToGameDir, FString& OutRepositoryRoot);
 
+/**
+ * Get Git config user.name & user.email
+ * @param	InPathToGitBinary	The path to the Git binary
+ * @param	InRepositoryRoot	The Git repository from where to run the command - usually the Game directory (can be empty)
+ * @param	OutUserName			Name of the Git user configured for this repository (or globaly)
+ * @param	OutEmailName		E-mail of the Git user configured for this repository (or globaly)
+ */
+void GetUserConfig(const FString& InPathToGitBinary, const FString& InRepositoryRoot, FString& OutUserName, FString& OutUserEmail);
+
+/**
+ * Get Git current checked-out branch
+ * @param	InPathToGitBinary	The path to the Git binary
+ * @param	InRepositoryRoot	The Git repository from where to run the command - usually the Game directory (can be empty)
+ * @param	OutBranchName		Name of the current checked-out branch (if any, ie. not in detached HEAD)
+ */
+void GetBranchName(const FString& InPathToGitBinary, const FString& InRepositoryRoot, FString& OutBranchName);
+
 /**
  * Run a Git command - output is a string TArray.
  *
diff --git a/Engine/Plugins/Developer/GitSourceControl/Source/GitSourceControl/Private/SGitSourceControlSettings.cpp b/Engine/Plugins/Developer/GitSourceControl/Source/GitSourceControl/Private/SGitSourceControlSettings.cpp
--- a/Engine/Plugins/Developer/GitSourceControl/Source/GitSourceControl/Private/SGitSourceControlSettings.cpp
+++ b/Engine/Plugins/Developer/GitSourceControl/Source/GitSourceControl/Private/SGitSourceControlSettings.cpp
@@ -12,12 +12,12 @@ void SGitSourceControlSettings::Construct(const FArguments& InArgs)
 
 	ChildSlot
 	[
-		SNew( SBorder )
-		.BorderImage( FEditorStyle::GetBrush("DetailsView.CategoryBottom") )
-		.Padding( FMargin( 0.0f, 3.0f, 0.0f, 0.0f ) )
+		SNew(SBorder)
+		.BorderImage( FEditorStyle::GetBrush("DetailsView.CategoryBottom"))
+		.Padding(FMargin(0.0f, 3.0f, 0.0f, 0.0f))
 		[
 			SNew(SVerticalBox)
-			+ SVerticalBox::Slot()
+			+SVerticalBox::Slot()
 			.FillHeight(1.0f)
 			.Padding(2.0f)
 			.VAlign(VAlign_Center)
@@ -50,7 +50,114 @@ void SGitSourceControlSettings::Construct(const FArguments& InArgs)
 						.Text(this, &SGitSourceControlSettings::GetBinaryPathText)
 						.ToolTipText(LOCTEXT("BinaryPathLabel_Tooltip", "Path to Git binary"))
 						.OnTextCommitted(this, &SGitSourceControlSettings::OnBinaryPathTextCommited)
-						.OnTextChanged(this, &SGitSourceControlSettings::OnBinaryPathTextCommited, ETextCommit::Default)
+						.Font(Font)
+					]
+				]
+			]
+			+SVerticalBox::Slot()
+			.FillHeight(1.0f)
+			.Padding(2.0f)
+			.VAlign(VAlign_Center)
+			[
+				SNew(SHorizontalBox)
+				+SHorizontalBox::Slot()
+				.FillWidth(1.0f)
+				[
+					SNew(SVerticalBox)
+					+SVerticalBox::Slot()
+					.FillHeight(1.0f)
+					.Padding(2.0f)
+					.VAlign(VAlign_Center)
+					[
+						SNew(STextBlock)
+						.Text(LOCTEXT("RepositoryRootLabel", "Root of the repository"))
+						.ToolTipText(LOCTEXT("RepositoryRootLabel_Tooltip", "Path to the root of the Git repository"))
+						.Font(Font)
+					]
+				]
+				+SHorizontalBox::Slot()
+				.FillWidth(2.0f)
+				[
+					SNew(SVerticalBox)
+					+SVerticalBox::Slot()
+					.FillHeight(1.0f)
+					.Padding(2.0f)
+					[
+						SNew(STextBlock)
+						.Text(this, &SGitSourceControlSettings::GetPathToRepositoryRoot)
+						.ToolTipText(LOCTEXT("RepositoryRootLabel_Tooltip", "Path to the root of the Git repository"))
+						.Font(Font)
+					]
+				]
+			]
+			+SVerticalBox::Slot()
+			.FillHeight(1.0f)
+			.Padding(2.0f)
+			.VAlign(VAlign_Center)
+			[
+				SNew(SHorizontalBox)
+				+SHorizontalBox::Slot()
+				.FillWidth(1.0f)
+				[
+					SNew(SVerticalBox)
+					+SVerticalBox::Slot()
+					.FillHeight(1.0f)
+					.Padding(2.0f)
+					.VAlign(VAlign_Center)
+					[
+						SNew(STextBlock)
+						.Text(LOCTEXT("GitUserName", "User Name"))
+						.ToolTipText(LOCTEXT("GitUserName_Tooltip", "User name configured for the Git repository"))
+						.Font(Font)
+					]
+				]
+				+SHorizontalBox::Slot()
+				.FillWidth(2.0f)
+				[
+					SNew(SVerticalBox)
+					+SVerticalBox::Slot()
+					.FillHeight(1.0f)
+					.Padding(2.0f)
+					[
+						SNew(STextBlock)
+						.Text(this, &SGitSourceControlSettings::GetUserName)
+						.ToolTipText(LOCTEXT("GitUserName_Tooltip", "User name configured for the Git repository"))
+						.Font(Font)
+					]
+				]
+			]
+			+SVerticalBox::Slot()
+			.FillHeight(1.0f)
+			.Padding(2.0f)
+			.VAlign(VAlign_Center)
+			[
+				SNew(SHorizontalBox)
+				+SHorizontalBox::Slot()
+				.FillWidth(1.0f)
+				[
+					SNew(SVerticalBox)
+					+SVerticalBox::Slot()
+					.FillHeight(1.0f)
+					.Padding(2.0f)
+					.VAlign(VAlign_Center)
+					[
+						SNew(STextBlock)
+						.Text(LOCTEXT("GitUserEmail", "E-Mail"))
+						.ToolTipText(LOCTEXT("GitUserEmail_Tooltip", "User e-mail configured for the Git repository"))
+						.Font(Font)
+					]
+				]
+				+SHorizontalBox::Slot()
+				.FillWidth(2.0f)
+				[
+					SNew(SVerticalBox)
+					+SVerticalBox::Slot()
+					.FillHeight(1.0f)
+					.Padding(2.0f)
+					[
+						SNew(STextBlock)
+						.Text(this, &SGitSourceControlSettings::GetUserEmail)
+						.ToolTipText(LOCTEXT("GitUserEmail_Tooltip", "User e-mail configured for the Git repository"))
 						.Font(Font)
 					]
 				]
@@ -72,4 +179,22 @@ void SGitSourceControlSettings::OnBinaryPathTextCommited(const FText& InText, ET
 	GitSourceControl.SaveSettings();
 }
 
+FText SGitSourceControlSettings::GetPathToRepositoryRoot() const
+{
+	FGitSourceControlModule& GitSourceControl = FModuleManager::LoadModuleChecked<FGitSourceControlModule>("GitSourceControl");
+	return FText::FromString(GitSourceControl.GetProvider().GetPathToRepositoryRoot());
+}
+
+FText SGitSourceControlSettings::GetUserName() const
+{
+	FGitSourceControlModule& GitSourceControl = FModuleManager::LoadModuleChecked<FGitSourceControlModule>("GitSourceControl");
+	return FText::FromString(GitSourceControl.GetProvider().GetUserName());
+}
+
+FText SGitSourceControlSettings::GetUserEmail() const
+{
+	FGitSourceControlModule& GitSourceControl = FModuleManager::LoadModuleChecked<FGitSourceControlModule>("GitSourceControl");
+	return FText::FromString(GitSourceControl.GetProvider().GetUserEmail());
+}
+
 #undef LOCTEXT_NAMESPACE
\ No newline at end of file
diff --git a/Engine/Plugins/Developer/GitSourceControl/Source/GitSourceControl/Private/SGitSourceControlSettings.h b/Engine/Plugins/Developer/GitSourceControl/Source/GitSourceControl/Private/SGitSourceControlSettings.h
--- a/Engine/Plugins/Developer/GitSourceControl/Source/GitSourceControl/Private/SGitSourceControlSettings.h
+++ b/Engine/Plugins/Developer/GitSourceControl/Source/GitSourceControl/Private/SGitSourceControlSettings.h
@@ -21,4 +21,9 @@ private:
 
 	/** Delegate to commit repository text to settings */
 	void OnBinaryPathTextCommited(const FText& InText, ETextCommit::Type InCommitType) const;
+
+	/** Delegate to get repository root, user name and email from provider */
+	FText GetPathToRepositoryRoot() const;
+	FText GetUserName() const;
+	FText GetUserEmail() const;
 };
\ No newline at end of file
