Remove Git source control plugin log spam

diff --git a/Engine/Plugins/Developer/GitSourceControl/Source/GitSourceControl/Private/GitSourceControlProvider.cpp b/Engine/Plugins/Developer/GitSourceControl/Source/GitSourceControl/Private/GitSourceControlProvider.cpp
--- a/Engine/Plugins/Developer/GitSourceControl/Source/GitSourceControl/Private/GitSourceControlProvider.cpp
+++ b/Engine/Plugins/Developer/GitSourceControl/Source/GitSourceControl/Private/GitSourceControlProvider.cpp
@@ -51,7 +51,6 @@ void FGitSourceControlProvider::CheckGitAvailability()
 	}
 	else
 	{
-		UE_LOG(LogSourceControl, Error, TEXT("Git not found"));
 		bGitAvailable = false;
 	}
 }
diff --git a/Engine/Plugins/Developer/GitSourceControl/Source/GitSourceControl/Private/GitSourceControlUtils.cpp b/Engine/Plugins/Developer/GitSourceControl/Source/GitSourceControl/Private/GitSourceControlUtils.cpp
--- a/Engine/Plugins/Developer/GitSourceControl/Source/GitSourceControl/Private/GitSourceControlUtils.cpp
+++ b/Engine/Plugins/Developer/GitSourceControl/Source/GitSourceControl/Private/GitSourceControlUtils.cpp
@@ -85,15 +85,6 @@ static bool RunCommandInternalRaw(const FString& InCommand, const FString& InPat
 // @todo: temporary debug logs
 //	UE_LOG(LogSourceControl, Log, TEXT("RunCommandInternalRaw: 'git %s'"), *FullCommand);
 	FPlatformProcess::ExecProcess(*InPathToGitBinary, *FullCommand, &ReturnCode, &OutResults, &OutErrors);
-	if(OutErrors.IsEmpty() && (ReturnCode == 0))
-	{
-		UE_LOG(LogSourceControl, Log, TEXT("'git %s':\n%s"), *LogableCommand, *OutResults);
-	}
-	else
-	{
-		// @todo filter here like RemoveRedundantErrors("' is outside repository") ?
-		UE_LOG(LogSourceControl, Error, TEXT("'git %s': ReturnCode=%d Errors='%s'"), *LogableCommand, ReturnCode, *OutErrors);
-	}
 
 	return ReturnCode == 0;
 }
@@ -189,15 +180,6 @@ bool CheckGitAvailability(const FString& InPathToGitBinary)
 		}
 	}
 
-	if(bGitAvailable)
-	{
-		UE_LOG(LogSourceControl, Log, TEXT("'%s' found"), *InPathToGitBinary);
-	}
-	else
-	{
-		UE_LOG(LogSourceControl, Log, TEXT("'%s' not found"), *InPathToGitBinary);
-	}
-
 	return bGitAvailable;
 }
 
