Temporarily disable Git plugin on Linux until UE-13628

diff --git a/Engine/Plugins/Developer/GitSourceControl/GitSourceControl.uplugin b/Engine/Plugins/Developer/GitSourceControl/GitSourceControl.uplugin
--- a/Engine/Plugins/Developer/GitSourceControl/GitSourceControl.uplugin
+++ b/Engine/Plugins/Developer/GitSourceControl/GitSourceControl.uplugin
@@ -17,7 +17,8 @@
 	[
 		{
 			"Name" : "GitSourceControl",
-			"Type" : "Editor"
+			"Type" : "Editor",
+			"BlacklistPlatforms" : [ "Linux" ]
 		}
 	]
 }
