PR #1043: Git plugin merge support (Contributed by [REDACTED_PERSON])

diff --git a/Engine/Plugins/Developer/GitSourceControl/Source/GitSourceControl/Private/GitSourceControlModule.h b/Engine/Plugins/Developer/GitSourceControl/Source/GitSourceControl/Private/GitSourceControlModule.h
--- a/Engine/Plugins/Developer/GitSourceControl/Source/GitSourceControl/Private/GitSourceControlModule.h
+++ b/Engine/Plugins/Developer/GitSourceControl/Source/GitSourceControl/Private/GitSourceControlModule.h
@@ -21,10 +21,9 @@ Written and contributed by [REDACTED_PERSON] ([REDACTED_EMAIL])
 - delete a file
 - checkin/commit a file (cannot handle atomically more than 20 files)
 - show current branch name in status text
-- solve a merge conflict
+- merge blueprints to solve a merge or rebase conflict
 
 ### What *cannot* be done presently
-- merge blueprints
 - initialize a new Git local repository ('git init') to manage your UE4 Game Project.
 - tags: implement ISourceControlLabel to manage git tags
 - .uproject file state si not visible in the current Editor
@@ -32,9 +31,10 @@ Written and contributed by [REDACTED_PERSON] ([REDACTED_EMAIL])
 - Pull/Fetch/Push are not in the current Editor workflow
 - Amend a commit is not in the current Editor workflow
 - configure user name & email ('git config user.name' & git config user.email')
-- git-annexe and/or git-media
+- Linux is not supported by the SourceControlProvider (#define SOURCE_CONTROL_WITH_SLATE	!PLATFORM_LINUX)
 
 ### Known issues:
+- global menu "Submit to source control" leads to many lines of Logs "is outside repository"
 - the Editor does not show deleted files (only when deleted externaly?)
 - the Editor does not show missing files
 - missing localisation for git specific messages
@@ -44,14 +44,11 @@ Written and contributed by [REDACTED_PERSON] ([REDACTED_EMAIL])
 - issue #11: Add the "Resolve" operation introduced in Editor 4.3
 - improve the 'Init' window text, hide it if connection is already done, auto connect
 - reverting an asset does not seem to update content in Editor! Issue in Editor?
-- file history show Revision as signed integer instead of hexadecimal SHA1
-- file history does not report file size
+- renaming a Blueprint in Editor leaves a tracker file, AND modify too much the asset to enable git to track its history through renaming
+- file history show Changelist as signed integer instead of hexadecimal SHA1
 - standard Editor commit dialog ask if user wants to "Keep Files Checked Out" => no use for Git or Mercurial CanCheckOut()==false
-- Windows only (64 bits) -> Mac compiles but needs testing/debugging (Linux source control is not supported by Editor)
 
  */
-
-
 class FGitSourceControlModule : public IModuleInterface
 {
 public:
diff --git a/Engine/Plugins/Developer/GitSourceControl/Source/GitSourceControl/Private/GitSourceControlOperations.cpp b/Engine/Plugins/Developer/GitSourceControl/Source/GitSourceControl/Private/GitSourceControlOperations.cpp
--- a/Engine/Plugins/Developer/GitSourceControl/Source/GitSourceControl/Private/GitSourceControlOperations.cpp
+++ b/Engine/Plugins/Developer/GitSourceControl/Source/GitSourceControl/Private/GitSourceControlOperations.cpp
@@ -177,23 +177,19 @@ bool FGitUpdateStatusWorker::Execute(FGitSourceControlCommand& InCommand)
 
 		if(Operation->ShouldUpdateHistory())
 		{
-			for(const auto& ItFile : InCommand.Files)
+			for(int32 Index = 0; Index < States.Num(); Index++)
 			{
-				TArray<FString> Results;
-				TArray<FString> Parameters;
-
-				Parameters.Add(TEXT("--max-count 100"));
-				Parameters.Add(TEXT("--follow")); // follow file renames
-				Parameters.Add(TEXT("--date=raw"));
-				Parameters.Add(TEXT("--name-status")); // relative filename at this revision, preceded by a status character
-
-				TArray<FString> Files;
-				Files.Add(*ItFile);
-
-				InCommand.bCommandSuccessful &= GitSourceControlUtils::RunCommand(TEXT("log"), InCommand.PathToGitBinary, InCommand.PathToRepositoryRoot, Parameters, Files, Results, InCommand.ErrorMessages);
+				FString& File = InCommand.Files[Index];
 				TGitSourceControlHistory History;
-				GitSourceControlUtils::ParseLogResults(Results, History);
-				Histories.Add(*ItFile, History);
+
+				if(States[Index].IsConflicted())
+				{
+					// In case of a merge conflict, we first need to get the tip of the "remote branch" (MERGE_HEAD)
+					GitSourceControlUtils::RunGetHistory(InCommand.PathToGitBinary, InCommand.PathToRepositoryRoot, File, true, InCommand.ErrorMessages, History);
+				}
+				// Get the history of the file in the current branch
+				InCommand.bCommandSuccessful &= GitSourceControlUtils::RunGetHistory(InCommand.PathToGitBinary, InCommand.PathToRepositoryRoot, File, false, InCommand.ErrorMessages, History);
+				Histories.Add(*File, History);
 			}
 		}
 	}
@@ -221,10 +217,10 @@ bool FGitUpdateStatusWorker::UpdateStates() const
 	FGitSourceControlProvider& Provider = GitSourceControl.GetProvider();
 
 	// add history, if any
-	for(const auto& ItHistory : Histories)
+	for(const auto& History : Histories)
 	{
-		TSharedRef<FGitSourceControlState, ESPMode::ThreadSafe> State = Provider.GetStateInternal(ItHistory.Key);
-		State->History = ItHistory.Value;
+		TSharedRef<FGitSourceControlState, ESPMode::ThreadSafe> State = Provider.GetStateInternal(History.Key);
+		State->History = History.Value;
 		State->TimeStamp = FDateTime::Now();
 		bUpdated = true;
 	}
@@ -267,8 +263,7 @@ bool FGitResolveWorker::Execute( class FGitSourceControlCommand& InCommand )
 	// mark the conflicting files as resolved:
 	{
 		TArray<FString> Results;
-		TArray<FString> Parameters;
-		InCommand.bCommandSuccessful = GitSourceControlUtils::RunCommand(TEXT("add"), InCommand.PathToGitBinary, InCommand.PathToRepositoryRoot, Parameters, InCommand.Files, Results, InCommand.ErrorMessages);
+		InCommand.bCommandSuccessful = GitSourceControlUtils::RunCommand(TEXT("add"), InCommand.PathToGitBinary, InCommand.PathToRepositoryRoot, TArray<FString>(), InCommand.Files, Results, InCommand.ErrorMessages);
 	}
 
 	// now update the status of our files
diff --git a/Engine/Plugins/Developer/GitSourceControl/Source/GitSourceControl/Private/GitSourceControlRevision.cpp b/Engine/Plugins/Developer/GitSourceControl/Source/GitSourceControl/Private/GitSourceControlRevision.cpp
--- a/Engine/Plugins/Developer/GitSourceControl/Source/GitSourceControl/Private/GitSourceControlRevision.cpp
+++ b/Engine/Plugins/Developer/GitSourceControl/Source/GitSourceControl/Private/GitSourceControlRevision.cpp
@@ -12,28 +12,31 @@
 bool FGitSourceControlRevision::Get( FString& InOutFilename ) const
 {
 	FGitSourceControlModule& GitSourceControl = FModuleManager::LoadModuleChecked<FGitSourceControlModule>("GitSourceControl");
-	FString PathToGitBinary = GitSourceControl.AccessSettings().GetBinaryPath();
-	FString PathToRepositoryRoot = GitSourceControl.GetProvider().GetPathToRepositoryRoot();
+	const FString PathToGitBinary = GitSourceControl.AccessSettings().GetBinaryPath();
+	const FString PathToRepositoryRoot = GitSourceControl.GetProvider().GetPathToRepositoryRoot();
 
 	// if a filename for the temp file wasn't supplied generate a unique-ish one
 	if(InOutFilename.Len() == 0)
 	{
-		// create the diff dir if we don't already have it (SVN wont)
+		// create the diff dir if we don't already have it (Git wont)
 		IFileManager::Get().MakeDirectory(*FPaths::DiffDir(), true);
-
-		// @todo Bug report: a counter is needed to avoid overlapping files; temp files are not (never?) released by Editor!
-		static int32 TempFileCount = 0;
-		FString TempFileName = FString::Printf(TEXT("%stemp-%d-%s-%s"), *FPaths::DiffDir(), TempFileCount++, *CommitId, *FPaths::GetCleanFilename(Filename));
+		// create a unique temp file name based on the unique commit Id
+		const FString TempFileName = FString::Printf(TEXT("%stemp-%s-%s"), *FPaths::DiffDir(), *CommitId, *FPaths::GetCleanFilename(Filename));
 		InOutFilename = FPaths::ConvertRelativePathToFull(TempFileName);
 	}
 
 	// Diff against the revision
-	FString Parameter = CommitId;
-	Parameter += TEXT(":");
-	Parameter += *Filename;
-
-	bool bCommandSuccessful = GitSourceControlUtils::RunDumpToFile(PathToGitBinary, PathToRepositoryRoot, Parameter, InOutFilename);
+	const FString Parameter = FString::Printf(TEXT("%s:%s"), *CommitId, *Filename);
 
+	bool bCommandSuccessful;
+	if(FPaths::FileExists(InOutFilename))
+	{
+		bCommandSuccessful = true; // if the temp file already exists, reuse it directly
+	}
+	else
+	{
+		bCommandSuccessful = GitSourceControlUtils::RunDumpToFile(PathToGitBinary, PathToRepositoryRoot, Parameter, InOutFilename);
+	}
 	return bCommandSuccessful;
 }
 
@@ -59,7 +62,7 @@ int32 FGitSourceControlRevision::GetRevisionNumber() const
 
 const FString& FGitSourceControlRevision::GetRevision() const
 {
-	return CommitId;
+	return ShortCommitId;
 }
 
 const FString& FGitSourceControlRevision::GetDescription() const
@@ -103,8 +106,7 @@ int32 FGitSourceControlRevision::GetCheckInIdentifier() const
 
 int32 FGitSourceControlRevision::GetFileSize() const
 {
-	// @todo git log does not give us the file size, but we could run a specific command
-	return 0;
+	return FileSize;
 }
 
 #undef LOCTEXT_NAMESPACE
diff --git a/Engine/Plugins/Developer/GitSourceControl/Source/GitSourceControl/Private/GitSourceControlRevision.h b/Engine/Plugins/Developer/GitSourceControl/Source/GitSourceControl/Private/GitSourceControlRevision.h
--- a/Engine/Plugins/Developer/GitSourceControl/Source/GitSourceControl/Private/GitSourceControlRevision.h
+++ b/Engine/Plugins/Developer/GitSourceControl/Source/GitSourceControl/Private/GitSourceControlRevision.h
@@ -34,12 +34,19 @@ public:
 	/** The filename this revision refers to */
 	FString Filename;
 
-	/** The SHA1 id of the commit this revision refers to */
+	/** The full hexadecimal SHA1 id of the commit this revision refers to */
 	FString CommitId;
 
-	/** The revision number */
+	/** The short hexadecimal SHA1 id (8 first hex char out of 40) of the commit: the string to display */
+	FString ShortCommitId;
+
+	/** The numeric value of the short SHA1 (8 first hex char out of 40) */
+	// @todo NOCOMMIT : the blueprint menu assume it to be in order !?
 	int32 RevisionNumber;
 
+	/** The SHA1 identifier of the file at this revision */
+	FString FileHash;
+
 	/** The description of this revision */
 	FString Description;
 
@@ -51,6 +58,9 @@ public:
 
 	/** The date this revision was made */
 	FDateTime Date;
+
+	/** The size of the file at this revision */
+	int32 FileSize;
 };
 
 /** History composed of the last 100 revisions of the file */
diff --git a/Engine/Plugins/Developer/GitSourceControl/Source/GitSourceControl/Private/GitSourceControlState.cpp b/Engine/Plugins/Developer/GitSourceControl/Source/GitSourceControl/Private/GitSourceControlState.cpp
--- a/Engine/Plugins/Developer/GitSourceControl/Source/GitSourceControl/Private/GitSourceControlState.cpp
+++ b/Engine/Plugins/Developer/GitSourceControl/Source/GitSourceControl/Private/GitSourceControlState.cpp
@@ -26,7 +26,7 @@ TSharedPtr<class ISourceControlRevision, ESPMode::ThreadSafe> FGitSourceControlS
 		}
 	}
 
-	return NULL;
+	return nullptr;
 }
 
 TSharedPtr<class ISourceControlRevision, ESPMode::ThreadSafe> FGitSourceControlState::FindHistoryRevision(const FString& InRevision) const
@@ -39,12 +39,20 @@ TSharedPtr<class ISourceControlRevision, ESPMode::ThreadSafe> FGitSourceControlS
 		}
 	}
 
-	return NULL;
+	return nullptr;
 }
 
 TSharedPtr<class ISourceControlRevision, ESPMode::ThreadSafe> FGitSourceControlState::GetBaseRevForMerge() const
 {
-	// @todo get revision of the merge-base (https://www.kernel.org/pub/software/scm/git/docs/git-merge-base.html)
+	for(const auto& Revision : History)
+	{
+		// look for the the SHA1 id of the file, not the commit id (revision)
+		if(Revision->FileHash == PendingMergeBaseFileHash)
+		{
+			return Revision;
+		}
+	}
+
 	return nullptr;
 }
 
diff --git a/Engine/Plugins/Developer/GitSourceControl/Source/GitSourceControl/Private/GitSourceControlState.h b/Engine/Plugins/Developer/GitSourceControl/Source/GitSourceControl/Private/GitSourceControlState.h
--- a/Engine/Plugins/Developer/GitSourceControl/Source/GitSourceControl/Private/GitSourceControlState.h
+++ b/Engine/Plugins/Developer/GitSourceControl/Source/GitSourceControl/Private/GitSourceControlState.h
@@ -68,6 +68,9 @@ public:
 	/** Filename on disk */
 	FString LocalFilename;
 
+	/** File Id with which our local revision diverged from the remote revision */
+	FString PendingMergeBaseFileHash;
+
 	/** State of the working copy */
 	EWorkingCopyState::Type WorkingCopyState;
 
diff --git a/Engine/Plugins/Developer/GitSourceControl/Source/GitSourceControl/Private/GitSourceControlUtils.cpp b/Engine/Plugins/Developer/GitSourceControl/Source/GitSourceControl/Private/GitSourceControlUtils.cpp
--- a/Engine/Plugins/Developer/GitSourceControl/Source/GitSourceControl/Private/GitSourceControlUtils.cpp
+++ b/Engine/Plugins/Developer/GitSourceControl/Source/GitSourceControl/Private/GitSourceControlUtils.cpp
@@ -221,7 +221,7 @@ void GetUserConfig(const FString& InPathToGitBinary, const FString& InRepository
 	TArray<FString> Parameters;
 	Parameters.Add(TEXT("user.name"));
 	bResults = RunCommandInternal(TEXT("config"), InPathToGitBinary, InRepositoryRoot, Parameters, TArray<FString>(), InfoMessages, ErrorMessages);
-	if (bResults && InfoMessages.Num() > 0)
+	if(bResults && InfoMessages.Num() > 0)
 	{
 		OutUserName = InfoMessages[0];
 	}
@@ -230,7 +230,7 @@ void GetUserConfig(const FString& InPathToGitBinary, const FString& InRepository
 	Parameters.Add(TEXT("user.email"));
 	InfoMessages.Reset();
 	bResults &= RunCommandInternal(TEXT("config"), InPathToGitBinary, InRepositoryRoot, Parameters, TArray<FString>(), InfoMessages, ErrorMessages);
-	if (bResults && InfoMessages.Num() > 0)
+	if(bResults && InfoMessages.Num() > 0)
 	{
 		OutUserEmail = InfoMessages[0];
 	}
@@ -246,7 +246,7 @@ void GetBranchName(const FString& InPathToGitBinary, const FString& InRepository
 	Parameters.Add(TEXT("--quiet"));		// no error message while in detached HEAD
 	Parameters.Add(TEXT("HEAD"));	
 	bResults = RunCommandInternal(TEXT("symbolic-ref"), InPathToGitBinary, InRepositoryRoot, Parameters, TArray<FString>(), InfoMessages, ErrorMessages);
-	if (bResults && InfoMessages.Num() > 0)
+	if(bResults && InfoMessages.Num() > 0)
 	{
 		OutBranchName = InfoMessages[0];
 	}
@@ -256,7 +256,7 @@ void GetBranchName(const FString& InPathToGitBinary, const FString& InRepository
 		Parameters.Add(TEXT("-1"));
 		Parameters.Add(TEXT("--format=\"%h\""));		// no error message while in detached HEAD
 		bResults = RunCommandInternal(TEXT("log"), InPathToGitBinary, InRepositoryRoot, Parameters, TArray<FString>(), InfoMessages, ErrorMessages);
-		if (bResults && InfoMessages.Num() > 0)
+		if(bResults && InfoMessages.Num() > 0)
 		{
 			OutBranchName = "HEAD detached at ";
 			OutBranchName += InfoMessages[0];
@@ -295,6 +295,7 @@ bool RunCommand(const FString& InCommand, const FString& InPathToGitBinary, cons
 	return bResult;
 }
 
+// Run a Git "commit" command by batches
 bool RunCommit(const FString& InPathToGitBinary, const FString& InRepositoryRoot, const TArray<FString>& InParameters, const TArray<FString>& InFiles, TArray<FString>& OutResults, TArray<FString>& OutErrorMessages)
 {
 	bool bResult = true;
@@ -371,17 +372,17 @@ private:
 };
 
 /**
-* Extract and interpret the file state from the given Git status result.
-* @see http://git-scm.com/docs/git-status
-* ' ' = unmodified
-* 'M' = modified
-* 'A' = added
-* 'D' = deleted
-* 'R' = renamed
-* 'C' = copied
-* 'U' = updated but unmerged
-* '?' = unknown/untracked
-* '!' = ignored
+ * Extract and interpret the file state from the given Git status result.
+ * @see http://git-scm.com/docs/git-status
+ * ' ' = unmodified
+ * 'M' = modified
+ * 'A' = added
+ * 'D' = deleted
+ * 'R' = renamed
+ * 'C' = copied
+ * 'U' = updated but unmerged
+ * '?' = unknown/untracked
+ * '!' = ignored
 */
 class FGitStatusParser
 {
@@ -440,15 +441,58 @@ public:
 	EWorkingCopyState::Type State;
 };
 
+/**
+ * Extract the status of a unmerged (conflict) file
+ *
+ * Example output of git ls-files --unmerged Content/Blueprints/BP_Test.uasset
+100644 d9b33098273547b57c0af314136f35b494e16dcb 1	Content/Blueprints/BP_Test.uasset
+100644 a14347dc3b589b78fb19ba62a7e3982f343718bc 2	Content/Blueprints/BP_Test.uasset
+100644 f3137a7167c840847cd7bd2bf07eefbfb2d9bcd2 3	Content/Blueprints/BP_Test.uasset
+ *
+ * 1: The "common ancestor" of the file (the version of the file that both the current and other branch originated from).
+ * 2: The version from the current branch (the master branch in this case).
+ * 3: The version from the other branch (the test branch)
+*/
+class FGitConflictStatusParser
+{
+public:
+	/** Parse the unmerge status: extract the base SHA1 identifier of the file */
+	FGitConflictStatusParser(const TArray<FString>& InResults)
+	{
+		const FString& FirstResult = InResults[0]; // 1: The common ancestor of merged branches
+		CommonAncestorFileId = FirstResult.Mid(7, 40);
+	}
+
+	FString CommonAncestorFileId;	/// SHA1 Id of the file (warning: not the commit Id)
+};
+
+/** Execute a command to get the details of a conflict */
+static void RunGetConflictStatus(const FString& InPathToGitBinary, const FString& InRepositoryRoot, const FString& InFile, FGitSourceControlState& InOutFileState)
+{
+	TArray<FString> ErrorMessages;
+	TArray<FString> Results;
+	TArray<FString> Files;
+	Files.Add(InFile);
+	TArray<FString> Parameters;
+	Parameters.Add(TEXT("--unmerged"));
+	bool bResult = RunCommandInternal(TEXT("ls-files"), InPathToGitBinary, InRepositoryRoot, Parameters, Files, Results, ErrorMessages);
+	if(bResult && Results.Num() == 3)
+	{
+		// Parse the unmerge status: extract the base revision (or the other branch?)
+		FGitConflictStatusParser ConflictStatus(Results);
+		InOutFileState.PendingMergeBaseFileHash = ConflictStatus.CommonAncestorFileId;
+	}
+}
+
 /** Parse the array of strings results of a 'git status' command
-*
-* Example git status results:
+ *
+ * Example git status results:
 M  Content/Textures/T_Perlin_Noise_M.uasset
 R  Content/Textures/T_Perlin_Noise_M.uasset -> Content/Textures/T_Perlin_Noise_M2.uasset
 ?? Content/Materials/M_Basic_Wall.uasset
 !! BasicCode.sln
 */
-static void ParseStatusResults(const TArray<FString>& InFiles, const TArray<FString>& InResults, TArray<FGitSourceControlState>& OutStates)
+static void ParseStatusResults(const FString& InPathToGitBinary, const FString& InRepositoryRoot, const TArray<FString>& InFiles, const TArray<FString>& InResults, TArray<FGitSourceControlState>& OutStates)
 {
 	// Iterate on all files explicitely listed in the command
 	for(const auto& File : InFiles)
@@ -461,6 +505,11 @@ static void ParseStatusResults(const TArray<FString>& InFiles, const TArray<FStr
 			// File found in status results; only the case for "changed" files
 			FGitStatusParser StatusParser(InResults[IdxResult]);
 			FileState.WorkingCopyState = StatusParser.State;
+			if(FileState.IsConflicted())
+			{
+				// In case of a conflict (unmerged file) get the base revision to merge
+				RunGetConflictStatus(InPathToGitBinary, InRepositoryRoot, File, FileState);
+			}
 		}
 		else
 		{
@@ -481,6 +530,7 @@ static void ParseStatusResults(const TArray<FString>& InFiles, const TArray<FStr
 	}
 }
 
+// Run a Git "status" command to update status of given files.
 bool RunUpdateStatus(const FString& InPathToGitBinary, const FString& InRepositoryRoot, const TArray<FString>& InFiles, TArray<FString>& OutErrorMessages, TArray<FGitSourceControlState>& OutStates)
 {
 	bool bResults = true;
@@ -516,7 +566,7 @@ bool RunUpdateStatus(const FString& InPathToGitBinary, const FString& InReposito
 		OutErrorMessages.Append(ErrorMessages);
 		if(bResult)
 		{
-			ParseStatusResults(Files.Value, Results, OutStates);
+			ParseStatusResults(InPathToGitBinary, InRepositoryRoot, Files.Value, Results, OutStates);
 		}
 		else
 		{
@@ -527,9 +577,7 @@ bool RunUpdateStatus(const FString& InPathToGitBinary, const FString& InReposito
 	return bResults;
 }
 
-/**
-* Run a Git show command to dump the binary content of a revision into a file.
-*/
+// Run a Git show command to dump the binary content of a revision into a file.
 bool RunDumpToFile(const FString& InPathToGitBinary, const FString& InRepositoryRoot, const FString& InParameter, const FString& InDumpFileName)
 {
 	bool bResult = false;
@@ -555,15 +603,12 @@ bool RunDumpToFile(const FString& InPathToGitBinary, const FString& InRepository
 	const bool bLaunchHidden = true;
 	const bool bLaunchReallyHidden = bLaunchHidden;
 
-	void* PipeRead = NULL;
-	void* PipeWrite = NULL;
+	void* PipeRead = nullptr;
+	void* PipeWrite = nullptr;
 
 	verify(FPlatformProcess::CreatePipe(PipeRead, PipeWrite));
 
-	// @todo temp debug log
-	//UE_LOG(LogSourceControl, Log, TEXT("RunDumpToFile: 'git %s'"), *FullCommand);
-	FProcHandle ProcessHandle = FPlatformProcess::CreateProc(*InPathToGitBinary, *FullCommand, bLaunchDetached, bLaunchHidden, bLaunchReallyHidden, NULL, 0, NULL, PipeWrite);
-	//UE_LOG(LogSourceControl, Log, TEXT("RunDumpToFile: ProcessHandle=%x"), ProcessHandle.Get());
+	FProcHandle ProcessHandle = FPlatformProcess::CreateProc(*InPathToGitBinary, *FullCommand, bLaunchDetached, bLaunchHidden, bLaunchReallyHidden, nullptr, 0, nullptr, PipeWrite);
 	if(ProcessHandle.IsValid())
 	{
 		FPlatformProcess::Sleep(0.01);
@@ -620,7 +665,7 @@ bool RunDumpToFile(const FString& InPathToGitBinary, const FString& InRepository
 * 'X' = unknown
 * 'B' = broken pairing
 */
-FString LogStatusToString(TCHAR InStatus)
+static FString LogStatusToString(TCHAR InStatus)
 {
 	switch(InStatus)
 	{
@@ -649,9 +694,10 @@ FString LogStatusToString(TCHAR InStatus)
 	return FString();
 }
 
-/** Parse the array of strings results of a 'git log' command
-*
-* Example git log results:
+/**
+ * Parse the array of strings results of a 'git log' command
+ *
+ * Example git log results:
 commit 97a4e7626681895e073aaefd68b8ac087db81b0b
 Author: [REDACTED_PERSON] <[REDACTED_EMAIL]>
 Date:   2014-2015-05-15 21:32:27 +0200
@@ -662,8 +708,8 @@ Date:   2014-2015-05-15 21:32:27 +0200
      - some <xml>
      - and strange characteres $*+
 
-M   Content/Blueprints/Blueprint_CeilingLight.uasset
-R100    Content/Textures/T_Concrete_Poured_D.uasset Content/Textures/T_Concrete_Poured_D2.uasset
+M	Content/Blueprints/Blueprint_CeilingLight.uasset
+R100	Content/Textures/T_Concrete_Poured_D.uasset Content/Textures/T_Concrete_Poured_D2.uasset
 
 commit 355f0df26ebd3888adbb558fd42bb8bd3e565000
 Author: [REDACTED_PERSON] <[REDACTED_EMAIL]>
@@ -671,10 +717,10 @@ Date:   2014-2015-05-12 11:28:14 +0200
 
     Testing git status, edit, and revert
 
-A    Content/Blueprints/Blueprint_CeilingLight.uasset
-C099    Content/Textures/T_Concrete_Poured_N.uasset Content/Textures/T_Concrete_Poured_N2.uasset
+A	Content/Blueprints/Blueprint_CeilingLight.uasset
+C099	Content/Textures/T_Concrete_Poured_N.uasset Content/Textures/T_Concrete_Poured_N2.uasset
 */
-void ParseLogResults(const TArray<FString>& InResults, TGitSourceControlHistory& OutHistory)
+static void ParseLogResults(const TArray<FString>& InResults, TGitSourceControlHistory& OutHistory)
 {
 	TSharedRef<FGitSourceControlRevision, ESPMode::ThreadSafe> SourceControlRevision = MakeShareable(new FGitSourceControlRevision);
 	for(const auto& Result : InResults)
@@ -684,15 +730,13 @@ void ParseLogResults(const TArray<FString>& InResults, TGitSourceControlHistory&
 			// End of the previous commit
 			if(SourceControlRevision->RevisionNumber != 0)
 			{
-				SourceControlRevision->Description += TEXT("\nCommit Id: ");
-				SourceControlRevision->Description += SourceControlRevision->CommitId;
 				OutHistory.Add(SourceControlRevision);
 
 				SourceControlRevision = MakeShareable(new FGitSourceControlRevision);
 			}
-			SourceControlRevision->CommitId = Result.RightChop(7);
-			FString ShortCommitId = SourceControlRevision->CommitId.Right(8); // Short revision ; first 8 hex characters (max that can hold a 32 bit integer)
-			SourceControlRevision->RevisionNumber = FParse::HexNumber(*ShortCommitId);
+			SourceControlRevision->CommitId = Result.RightChop(7); // Full commit SHA1 hexadecimal string
+			SourceControlRevision->ShortCommitId = SourceControlRevision->CommitId.Left(8); // Short revision ; first 8 hex characters (max that can hold a 32 bit integer)
+			SourceControlRevision->RevisionNumber = FParse::HexNumber(*SourceControlRevision->ShortCommitId);
 		}
 		else if(Result.StartsWith(TEXT("Author: "))) // Author name & email
 		{
@@ -715,22 +759,98 @@ void ParseLogResults(const TArray<FString>& InResults, TGitSourceControlHistory&
 			SourceControlRevision->Description += Result.RightChop(4);
 			SourceControlRevision->Description += TEXT("\n");
 		}
-		else // List of modified files, starting with a uppercase status letter ("A"/"M"...)
+		else // Name of the file, starting with an uppercase status letter ("A"/"M"...)
 		{
-			TCHAR Status = Result[0];
+			const TCHAR Status = Result[0];
 			SourceControlRevision->Action = LogStatusToString(Status); // Readable action string ("Added", Modified"...) instead of "A"/"M"...
-			SourceControlRevision->Filename = Result.RightChop(2); // relative filename
+			// Take care of special case for Renamed/Copied file: extract the second filename after second tabulation
+			int32 IdxTab;
+			if(Result.FindLastChar('\t', IdxTab))
+			{
+				SourceControlRevision->Filename = Result.RightChop(IdxTab + 1); // relative filename
+			}
 		}
 	}
 	// End of the last commit
 	if(SourceControlRevision->RevisionNumber != 0)
 	{
-		SourceControlRevision->Description += TEXT("\nCommit Id: ");
-		SourceControlRevision->Description += SourceControlRevision->CommitId;
 		OutHistory.Add(SourceControlRevision);
 	}
 }
 
+
+/**
+ * Extract the SHA1 identifier and size of a blob (file) from a Git "ls-tree" command.
+ *
+ * Example output for the command git ls-tree --long 7fdaeb2 Content/Blueprints/BP_Test.uasset
+100644 blob a14347dc3b589b78fb19ba62a7e3982f343718bc   70731	Content/Blueprints/BP_Test.uasset
+*/
+class FGitLsTreeParser
+{
+public:
+	/** Parse the unmerge status: extract the base SHA1 identifier of the file */
+	FGitLsTreeParser(const TArray<FString>& InResults)
+	{
+		const FString& FirstResult = InResults[0];
+		FileHash = FirstResult.Mid(12, 40);
+		int32 IdxTab;
+		if(FirstResult.FindChar('\t', IdxTab))
+		{
+			const FString SizeString = FirstResult.Mid(53, IdxTab - 53);
+			FileSize = FCString::Atoi(*SizeString);
+		}
+	}
+
+	FString FileHash;	/// SHA1 Id of the file (warning: not the commit Id)
+	int32	FileSize;	/// Size of the file (in bytes)
+};
+
+// Run a Git "log" command and parse it.
+bool RunGetHistory(const FString& InPathToGitBinary, const FString& InRepositoryRoot, const FString& InFile, bool bMergeConflict, TArray<FString>& OutErrorMessages, TGitSourceControlHistory& OutHistory)
+{
+	bool bResults;
+	{
+		TArray<FString> Results;
+		TArray<FString> Parameters;
+		Parameters.Add(bMergeConflict?TEXT("--max-count 1"):TEXT("--max-count 100"));
+		Parameters.Add(TEXT("--follow")); // follow file renames
+		Parameters.Add(TEXT("--date=raw"));
+		Parameters.Add(TEXT("--name-status")); // relative filename at this revision, preceded by a status character
+		TArray<FString> Files;
+		Files.Add(*InFile);
+		if(bMergeConflict)
+		{
+			// In case of a merge conflict, we also need to get the tip of the "remote branch" (MERGE_HEAD) before the log of the "current branch" (HEAD)
+			// @todo does not work for a cherry-pick! Test for a rebase.
+			Parameters.Add(TEXT("MERGE_HEAD"));
+		}
+		bResults = RunCommand(TEXT("log"), InPathToGitBinary, InRepositoryRoot, Parameters, Files, Results, OutErrorMessages);
+		if(bResults)
+		{
+			ParseLogResults(Results, OutHistory);
+		}
+	}
+	for(auto& Revision : OutHistory)
+	{
+		// Get file (blob) sha1 id and size
+		TArray<FString> Results;
+		TArray<FString> Parameters;
+		Parameters.Add(TEXT("--long")); // Show object size of blob (file) entries.
+		Parameters.Add(Revision->GetRevision());
+		TArray<FString> Files;
+		Files.Add(*Revision->GetFilename());
+		bResults &= RunCommand(TEXT("ls-tree"), InPathToGitBinary, InRepositoryRoot, Parameters, Files, Results, OutErrorMessages);
+		if(bResults && Results.Num())
+		{
+			FGitLsTreeParser LsTree(Results);
+			Revision->FileHash = LsTree.FileHash;
+			Revision->FileSize = LsTree.FileSize;
+		}
+	}
+
+	return bResults;
+}
+
 bool UpdateCachedStates(const TArray<FGitSourceControlState>& InStates)
 {
 	FGitSourceControlModule& GitSourceControl = FModuleManager::LoadModuleChecked<FGitSourceControlModule>( "GitSourceControl" );
@@ -743,6 +863,7 @@ bool UpdateCachedStates(const TArray<FGitSourceControlState>& InStates)
 		if(State->WorkingCopyState != InState.WorkingCopyState)
 		{
 			State->WorkingCopyState = InState.WorkingCopyState;
+			State->PendingMergeBaseFileHash = InState.PendingMergeBaseFileHash;
 		//	State->TimeStamp = InState.TimeStamp; // @todo Bug report: Workaround a bug with the Source Control Module not updating file state after a "Save"
 			NbStatesUpdated++;
 		}
diff --git a/Engine/Plugins/Developer/GitSourceControl/Source/GitSourceControl/Private/GitSourceControlUtils.h b/Engine/Plugins/Developer/GitSourceControl/Source/GitSourceControl/Private/GitSourceControlUtils.h
--- a/Engine/Plugins/Developer/GitSourceControl/Source/GitSourceControl/Private/GitSourceControlUtils.h
+++ b/Engine/Plugins/Developer/GitSourceControl/Source/GitSourceControl/Private/GitSourceControlUtils.h
@@ -84,7 +84,7 @@ void GetBranchName(const FString& InPathToGitBinary, const FString& InRepository
 bool RunCommand(const FString& InCommand, const FString& InPathToGitBinary, const FString& InRepositoryRoot, const TArray<FString>& InParameters, const TArray<FString>& InFiles, TArray<FString>& OutResults, TArray<FString>& OutErrorMessages);
 
 /**
- * Run a Git commit command by batches.
+ * Run a Git "commit" command by batches.
  *
  * @param	InPathToGitBinary	The path to the Git binary
  * @param	InRepositoryRoot	The Git repository from where to run the command - usually the Game directory (can be empty)
@@ -96,7 +96,7 @@ bool RunCommand(const FString& InCommand, const FString& InPathToGitBinary, cons
 bool RunCommit(const FString& InPathToGitBinary, const FString& InRepositoryRoot, const TArray<FString>& InParameters, const TArray<FString>& InFiles, TArray<FString>& OutResults, TArray<FString>& OutErrorMessages);
 
 /**
- * Run a Git status command and parse it.
+ * Run a Git "status" command and parse it.
  *
  * @param	InPathToGitBinary	The path to the Git binary
  * @param	InRepositoryRoot	The Git repository from where to run the command - usually the Game directory (can be empty)
@@ -107,7 +107,7 @@ bool RunCommit(const FString& InPathToGitBinary, const FString& InRepositoryRoot
 bool RunUpdateStatus(const FString& InPathToGitBinary, const FString& InRepositoryRoot, const TArray<FString>& InFiles, TArray<FString>& OutErrorMessages, TArray<FGitSourceControlState>& OutStates);
 
 /**
- * Run a Git show command to dump the binary content of a revision into a file.
+ * Run a Git "show" command to dump the binary content of a revision into a file.
  *
  * @param	InPathToGitBinary	The path to the Git binary
  * @param	InRepositoryRoot	The Git repository from where to run the command - usually the Game directory (can be empty)
@@ -118,11 +118,16 @@ bool RunUpdateStatus(const FString& InPathToGitBinary, const FString& InReposito
 bool RunDumpToFile(const FString& InPathToGitBinary, const FString& InRepositoryRoot, const FString& InParameter, const FString& InDumpFileName);
 
 /**
- * Parse the array of strings results of a 'git log' command
- * @param	InResults			The results (from StdOut) as an array per-line
+ * Run a Git "log" command and parse it.
+ *
+ * @param	InPathToGitBinary	The path to the Git binary
+ * @param	InRepositoryRoot	The Git repository from where to run the command - usually the Game directory (can be empty)
+ * @param	InFile				The file to be operated on
+ * @param	bMergeConflict		In case of a merge conflict, we also need to get the tip of the "remote branch" (MERGE_HEAD) before the log of the "current branch" (HEAD)
+ * @param	OutErrorMessages	Any errors (from StdErr) as an array per-line
  * @param	OutHistory			The history of the file
  */
-void ParseLogResults(const TArray<FString>& InResults, TGitSourceControlHistory& OutHistory);
+bool RunGetHistory(const FString& InPathToGitBinary, const FString& InRepositoryRoot, const FString& InFile, bool bMergeConflict, TArray<FString>& OutErrorMessages, TGitSourceControlHistory& OutHistory);
 
 /**
  * Helper function for various commands to update cached states.
