PR #1056: (Git plugin) Option to initialize new local repository for current project (Contributed by [REDACTED_PERSON])

diff --git a/Engine/Plugins/Developer/GitSourceControl/Source/GitSourceControl/GitSourceControl.Build.cs b/Engine/Plugins/Developer/GitSourceControl/Source/GitSourceControl/GitSourceControl.Build.cs
--- a/Engine/Plugins/Developer/GitSourceControl/Source/GitSourceControl/GitSourceControl.Build.cs
+++ b/Engine/Plugins/Developer/GitSourceControl/Source/GitSourceControl/GitSourceControl.Build.cs
@@ -11,6 +11,7 @@ public class GitSourceControl : ModuleRules
 				"Core",
 				"Slate",
 				"SlateCore",
+				"InputCore",
 				"EditorStyle",
 				"SourceControl",
 			}
diff --git a/Engine/Plugins/Developer/GitSourceControl/Source/GitSourceControl/Private/GitSourceControlModule.h b/Engine/Plugins/Developer/GitSourceControl/Source/GitSourceControl/Private/GitSourceControlModule.h
--- a/Engine/Plugins/Developer/GitSourceControl/Source/GitSourceControl/Private/GitSourceControlModule.h
+++ b/Engine/Plugins/Developer/GitSourceControl/Source/GitSourceControl/Private/GitSourceControlModule.h
@@ -13,6 +13,7 @@ UE4GitPlugin is a simple Git Source Control Plugin for Unreal Engine
 Written and contributed by [REDACTED_PERSON] ([REDACTED_EMAIL])
 
 ### Supported features
+- initialize a new Git local repository ('git init') to manage your UE4 Game Project.
 - display status icons to show modified/added/deleted/untracked files
 - show history of a file
 - diff against depot or between previous versions of a file
@@ -24,7 +25,6 @@ Written and contributed by [REDACTED_PERSON] ([REDACTED_EMAIL])
 - merge blueprints to solve a merge or rebase conflict
 
 ### What *cannot* be done presently
-- initialize a new Git local repository ('git init') to manage your UE4 Game Project.
 - tags: implement ISourceControlLabel to manage git tags
 - .uproject file state si not visible in the current Editor
 - Branch is not in the current Editor workflow (but on Epic Roadmap)
@@ -32,6 +32,7 @@ Written and contributed by [REDACTED_PERSON] ([REDACTED_EMAIL])
 - Amend a commit is not in the current Editor workflow
 - configure user name & email ('git config user.name' & git config user.email')
 - Linux is not supported by the SourceControlProvider (#define SOURCE_CONTROL_WITH_SLATE	!PLATFORM_LINUX)
+- git-annex and/or git-media
 
 ### Known issues:
 - global menu "Submit to source control" leads to many lines of Logs "is outside repository"
@@ -42,7 +43,6 @@ Written and contributed by [REDACTED_PERSON] ([REDACTED_EMAIL])
 - displaying states of 'Engine' assets (also needs management of 'out of tree' files)
 - issue #22: A Move/Rename leaves a redirector file behind
 - issue #11: Add the "Resolve" operation introduced in Editor 4.3
-- improve the 'Init' window text, hide it if connection is already done, auto connect
 - reverting an asset does not seem to update content in Editor! Issue in Editor?
 - renaming a Blueprint in Editor leaves a tracker file, AND modify too much the asset to enable git to track its history through renaming
 - file history show Changelist as signed integer instead of hexadecimal SHA1
diff --git a/Engine/Plugins/Developer/GitSourceControl/Source/GitSourceControl/Private/GitSourceControlProvider.cpp b/Engine/Plugins/Developer/GitSourceControl/Source/GitSourceControl/Private/GitSourceControlProvider.cpp
--- a/Engine/Plugins/Developer/GitSourceControl/Source/GitSourceControl/Private/GitSourceControlProvider.cpp
+++ b/Engine/Plugins/Developer/GitSourceControl/Source/GitSourceControl/Private/GitSourceControlProvider.cpp
@@ -35,17 +35,18 @@ void FGitSourceControlProvider::CheckGitAvailability()
 			// Find the path to the root Git directory (if any)
 			const FString PathToGameDir = FPaths::ConvertRelativePathToFull(FPaths::GameDir());
 			const bool bRepositoryFound = GitSourceControlUtils::FindRootDirectory(PathToGameDir, PathToRepositoryRoot);
-			if(bRepositoryFound)
+			// Get user name & email (of the repository, else from the global Git config)
+			GitSourceControlUtils::GetUserConfig(PathToGitBinary, PathToRepositoryRoot, UserName, UserEmail);
+			if (bRepositoryFound)
 			{
-				// Get user name & email
-				GitSourceControlUtils::GetUserConfig(PathToGitBinary, PathToRepositoryRoot, UserName, UserEmail);
 				// Get branch name
 				GitSourceControlUtils::GetBranchName(PathToGitBinary, PathToRepositoryRoot, BranchName);
+				bGitRepositoryFound = true;
 			}
 			else
 			{
 				UE_LOG(LogSourceControl, Error, TEXT("'%s' is not part of a Git repository"), *FPaths::GameDir());
-				bGitAvailable = false;
+				bGitRepositoryFound = false;
 			}
 		}
 	}
@@ -89,14 +90,16 @@ FText FGitSourceControlProvider::GetStatusText() const
 	return FText::Format( NSLOCTEXT("Status", "Provider: Git\nEnabledLabel", "Repository: {RepositoryName}\nBranch: {BranchName}\nUser: {UserName}\nE-mail: {UserEmail}"), Args );
 }
 
+/** Quick check if source control is enabled */
 bool FGitSourceControlProvider::IsEnabled() const
 {
-	return bGitAvailable;
+	return bGitRepositoryFound;
 }
 
+/** Quick check if source control is available for use (useful for server-based providers) */
 bool FGitSourceControlProvider::IsAvailable() const
 {
-	return bGitAvailable;
+	return bGitRepositoryFound;
 }
 
 const FName& FGitSourceControlProvider::GetName(void) const
diff --git a/Engine/Plugins/Developer/GitSourceControl/Source/GitSourceControl/Private/GitSourceControlProvider.h b/Engine/Plugins/Developer/GitSourceControl/Source/GitSourceControl/Private/GitSourceControlProvider.h
--- a/Engine/Plugins/Developer/GitSourceControl/Source/GitSourceControl/Private/GitSourceControlProvider.h
+++ b/Engine/Plugins/Developer/GitSourceControl/Source/GitSourceControl/Private/GitSourceControlProvider.h
@@ -21,6 +21,7 @@ public:
 	/** Constructor */
 	FGitSourceControlProvider() 
 		: bGitAvailable(false)
+		, bGitRepositoryFound(false)
 	{
 	}
 
@@ -53,6 +54,12 @@ public:
 	 */
 	void CheckGitAvailability();
 
+	/** Is git binary found and working. */
+	inline bool IsGitAvailable() const
+	{
+		return bGitAvailable;
+	}
+
 	/** Get the path to the root of the Git repository: can be the GameDir itself, or any parent directory */
 	inline const FString& GetPathToRepositoryRoot() const
 	{
@@ -85,6 +92,9 @@ private:
 	/** Is git binary found and working. */
 	bool bGitAvailable;
 
+	/** Is git repository found. */
+	bool bGitRepositoryFound;
+
 	/** Helper function for Execute() */
 	TSharedPtr<class IGitSourceControlWorker, ESPMode::ThreadSafe> CreateWorker(const FName& InOperationName) const;
 
diff --git a/Engine/Plugins/Developer/GitSourceControl/Source/GitSourceControl/Private/SGitSourceControlSettings.cpp b/Engine/Plugins/Developer/GitSourceControl/Source/GitSourceControl/Private/SGitSourceControlSettings.cpp
--- a/Engine/Plugins/Developer/GitSourceControl/Source/GitSourceControl/Private/SGitSourceControlSettings.cpp
+++ b/Engine/Plugins/Developer/GitSourceControl/Source/GitSourceControl/Private/SGitSourceControlSettings.cpp
@@ -3,6 +3,7 @@
 #include "GitSourceControlPrivatePCH.h"
 #include "SGitSourceControlSettings.h"
 #include "GitSourceControlModule.h"
+#include "GitSourceControlUtils.h"
 
 #define LOCTEXT_NAMESPACE "SGitSourceControlSettings"
 
@@ -10,6 +11,8 @@ void SGitSourceControlSettings::Construct(const FArguments& InArgs)
 {
 	FSlateFontInfo Font = FEditorStyle::GetFontStyle(TEXT("SourceControl.LoginWindow.Font"));
 
+	bAutoCreateGitIgnore = true;
+
 	ChildSlot
 	[
 		SNew(SBorder)
@@ -162,6 +165,69 @@ void SGitSourceControlSettings::Construct(const FArguments& InArgs)
 					]
 				]
 			]
+			+SVerticalBox::Slot()
+			.FillHeight(1.0f)
+			.Padding(2.0f)
+			.VAlign(VAlign_Center)
+			[
+				SNew(SHorizontalBox)
+				.Visibility(this, &SGitSourceControlSettings::CanInitializeGitRepository)
+				+SHorizontalBox::Slot()
+				.FillWidth(1.0f)
+				[
+					SNew(SVerticalBox)
+					+SVerticalBox::Slot()
+					.FillHeight(2.0f)
+					.Padding(2.0f)
+					.VAlign(VAlign_Center)
+					.AutoHeight()
+					[
+						SNew(SButton)
+						.Text(LOCTEXT("GitInitRepository", "Initialize project with Git"))
+						.ToolTipText(LOCTEXT("GitInitRepository_Tooltip", "Initialize current project as a new Git repository"))
+						.OnClicked(this, &SGitSourceControlSettings::OnClickedInitializeGitRepository)
+						.HAlign(HAlign_Center)
+					]
+				]
+			]
+			+SVerticalBox::Slot()
+			.FillHeight(1.0f)
+			.Padding(2.0f)
+			.VAlign(VAlign_Center)
+			[
+				SNew(SHorizontalBox)
+				.Visibility(this, &SGitSourceControlSettings::CanInitializeGitRepository)
+				+SHorizontalBox::Slot()
+				.FillWidth(0.04f)
+				[
+					SNew(SVerticalBox)
+					+SVerticalBox::Slot()
+					.FillHeight(1.0f)
+					.Padding(2.0f)
+					.VAlign(VAlign_Center)
+					[
+						SNew(SCheckBox)
+						.ToolTipText(LOCTEXT("CreateGitIgnore_Tooltip", "Create and add a standard '.gitignore' file"))
+						.IsChecked(true)
+						.OnCheckStateChanged(this, &SGitSourceControlSettings::OnCheckedCreateGitIgnore)
+					]
+				]
+				+SHorizontalBox::Slot()
+				.FillWidth(0.96f)
+				[
+					SNew(SVerticalBox)
+					+SVerticalBox::Slot()
+					.FillHeight(1.0f)
+					.Padding(2.0f)
+					.VAlign(VAlign_Center)
+					[
+						SNew(STextBlock)
+						.Text(LOCTEXT("CreateGitIgnore", "Add a .gitignore file"))
+						.ToolTipText(LOCTEXT("CreateGitIgnore_Tooltip", "Create and add a standard '.gitignore' file"))
+						.Font(Font)
+					]
+				]
+			]
 		]
 	];
 }
@@ -197,4 +263,49 @@ FText SGitSourceControlSettings::GetUserEmail() const
 	return FText::FromString(GitSourceControl.GetProvider().GetUserEmail());
 }
 
+EVisibility SGitSourceControlSettings::CanInitializeGitRepository() const
+{
+	FGitSourceControlModule& GitSourceControl = FModuleManager::LoadModuleChecked<FGitSourceControlModule>("GitSourceControl");
+	const bool bGitAvailable = GitSourceControl.GetProvider().IsGitAvailable();
+	const bool bGitRepositoryFound = GitSourceControl.GetProvider().IsEnabled();
+	return (bGitAvailable && !bGitRepositoryFound) ? EVisibility::Visible : EVisibility::Collapsed;
+}
+
+FReply SGitSourceControlSettings::OnClickedInitializeGitRepository()
+{
+	FGitSourceControlModule& GitSourceControl = FModuleManager::LoadModuleChecked<FGitSourceControlModule>("GitSourceControl");
+	TArray<FString> InfoMessages;
+	TArray<FString> ErrorMessages;
+	const FString& PathToGitBinary = GitSourceControl.AccessSettings().GetBinaryPath();
+	const FString PathToGameDir = FPaths::ConvertRelativePathToFull(FPaths::GameDir());
+	GitSourceControlUtils::RunCommand(TEXT("init"), PathToGitBinary, PathToGameDir, TArray<FString>(), TArray<FString>(), InfoMessages, ErrorMessages);
+	// Check the new repository status to enable connection
+	GitSourceControl.GetProvider().CheckGitAvailability();
+	if(GitSourceControl.GetProvider().IsEnabled())
+	{
+		TArray<FString> ProjectFiles;
+		const FString ProjectFile = FPaths::GetCleanFilename(FPaths::GetProjectFilePath());
+		ProjectFiles.Add(ProjectFile);
+		ProjectFiles.Add(TEXT("Config"));
+		if(bAutoCreateGitIgnore)
+		{
+			// Create a standard ".gitignore" file with common patterns for a typical Blueprint & C++ project
+			const FString Filename = FString::Printf(TEXT("%s.gitignore"), *PathToGameDir);
+			const FString GitIgnoreContent = TEXT("Binaries\nDerivedDataCache\nIntermediate\nSaved\n*.sdf\n*.sln\n*.suo");
+			if(FFileHelper::SaveStringToFile(GitIgnoreContent, *Filename, FFileHelper::EEncodingOptions::ForceUTF8WithoutBOM))
+			{
+				ProjectFiles.Add(TEXT(".gitignore"));
+			}
+		}
+		// Add .uproject and Config/ files (and .gitignore if any)
+		GitSourceControlUtils::RunCommand(TEXT("add"), PathToGitBinary, PathToGameDir, TArray<FString>(), ProjectFiles, InfoMessages, ErrorMessages);
+	}
+	return FReply::Handled();
+}
+
+void SGitSourceControlSettings::OnCheckedCreateGitIgnore(ECheckBoxState NewCheckedState)
+{
+	bAutoCreateGitIgnore = (NewCheckedState == ECheckBoxState::Checked);
+}
+
 #undef LOCTEXT_NAMESPACE
\ No newline at end of file
diff --git a/Engine/Plugins/Developer/GitSourceControl/Source/GitSourceControl/Private/SGitSourceControlSettings.h b/Engine/Plugins/Developer/GitSourceControl/Source/GitSourceControl/Private/SGitSourceControlSettings.h
--- a/Engine/Plugins/Developer/GitSourceControl/Source/GitSourceControl/Private/SGitSourceControlSettings.h
+++ b/Engine/Plugins/Developer/GitSourceControl/Source/GitSourceControl/Private/SGitSourceControlSettings.h
@@ -26,4 +26,10 @@ private:
 	FText GetPathToRepositoryRoot() const;
 	FText GetUserName() const;
 	FText GetUserEmail() const;
+
+	/** Delegate to initialize a new Git repository */
+	EVisibility CanInitializeGitRepository() const;
+	FReply OnClickedInitializeGitRepository();
+	void OnCheckedCreateGitIgnore(ECheckBoxState NewCheckedState);
+	bool bAutoCreateGitIgnore;
 };
\ No newline at end of file
