PR #1096: (Git plugin) restore connect error popup message (Contributed by [REDACTED_PERSON])

diff --git a/Engine/Plugins/Developer/GitSourceControl/GitSourceControl.uplugin b/Engine/Plugins/Developer/GitSourceControl/GitSourceControl.uplugin
--- a/Engine/Plugins/Developer/GitSourceControl/GitSourceControl.uplugin
+++ b/Engine/Plugins/Developer/GitSourceControl/GitSourceControl.uplugin
@@ -2,9 +2,9 @@
 	"FileVersion" : 3,
 	
 	"FriendlyName" : "Git",
-	"Version" : 12,
-	"VersionName" : "1.0",
-	"EngineVersion" : "4.7.0",
+	"Version" : 13,
+	"VersionName" : "1.1",
+	"EngineVersion" : "4.8.0",
 	"Description" : "Git source control management",
 	"Category" : "Source Control",
 	"CreatedBy" : "SRombauts and Epic Games, Inc.",
diff --git a/Engine/Plugins/Developer/GitSourceControl/Source/GitSourceControl/Private/GitSourceControlOperations.cpp b/Engine/Plugins/Developer/GitSourceControl/Source/GitSourceControl/Private/GitSourceControlOperations.cpp
--- a/Engine/Plugins/Developer/GitSourceControl/Source/GitSourceControl/Private/GitSourceControlOperations.cpp
+++ b/Engine/Plugins/Developer/GitSourceControl/Source/GitSourceControl/Private/GitSourceControlOperations.cpp
@@ -21,11 +21,7 @@ bool FGitConnectWorker::Execute(FGitSourceControlCommand& InCommand)
 	InCommand.bCommandSuccessful = GitSourceControlUtils::RunCommand(TEXT("status"), InCommand.PathToGitBinary, InCommand.PathToRepositoryRoot, TArray<FString>(), TArray<FString>(), InCommand.InfoMessages, InCommand.ErrorMessages);
 	if(!InCommand.bCommandSuccessful || InCommand.ErrorMessages.Num() > 0 || InCommand.InfoMessages.Num() == 0)
 	{
-		// @todo popup to propose to initialize the git repository "git init + .gitignore" (easier in the provider!)
-		// then "Connect" could in fact do the "git init"
-		StaticCastSharedRef<FConnect>(InCommand.Operation)->SetErrorText(LOCTEXT("NotAWorkingCopyError", "Project is not part of a Git working copy."));
-		// @todo Double error messages (and displayed in reverse order): Perforce distinguish the two errors
-		InCommand.ErrorMessages.Add(LOCTEXT("NotAWorkingCopyErrorHelp", "You should check out a working copy into your project directory.").ToString());
+		StaticCastSharedRef<FConnect>(InCommand.Operation)->SetErrorText(LOCTEXT("NotAGitRepository", "Failed to enable Git source control. You need to initialize the project as a Git repository first."));
 		InCommand.bCommandSuccessful = false;
 	}
 
diff --git a/Engine/Plugins/Developer/GitSourceControl/Source/GitSourceControl/Private/GitSourceControlProvider.cpp b/Engine/Plugins/Developer/GitSourceControl/Source/GitSourceControl/Private/GitSourceControlProvider.cpp
--- a/Engine/Plugins/Developer/GitSourceControl/Source/GitSourceControl/Private/GitSourceControlProvider.cpp
+++ b/Engine/Plugins/Developer/GitSourceControl/Source/GitSourceControl/Private/GitSourceControlProvider.cpp
@@ -165,7 +165,7 @@ void FGitSourceControlProvider::UnregisterSourceControlStateChanged_Handle( FDel
 
 ECommandResult::Type FGitSourceControlProvider::Execute( const TSharedRef<ISourceControlOperation, ESPMode::ThreadSafe>& InOperation, const TArray<FString>& InFiles, EConcurrency::Type InConcurrency, const FSourceControlOperationComplete& InOperationCompleteDelegate )
 {
-	if(!IsEnabled())
+	if(!IsEnabled() && !(InOperation->GetName() == "Connect")) // Only Connect operation allowed while not Enabled (Connected)
 	{
 		return ECommandResult::Failed;
 	}
diff --git a/Engine/Plugins/Developer/GitSourceControl/Source/GitSourceControl/Private/GitSourceControlUtils.cpp b/Engine/Plugins/Developer/GitSourceControl/Source/GitSourceControl/Private/GitSourceControlUtils.cpp
--- a/Engine/Plugins/Developer/GitSourceControl/Source/GitSourceControl/Private/GitSourceControlUtils.cpp
+++ b/Engine/Plugins/Developer/GitSourceControl/Source/GitSourceControl/Private/GitSourceControlUtils.cpp
@@ -209,7 +209,10 @@ bool FindRootDirectory(const FString& InPathToGameDir, FString& OutRepositoryRoo
 			}
 		}
 	}
-
+	if (!bFound)
+	{
+		OutRepositoryRoot = InPathToGameDir; // If not found, return the GameDir as best possible root.
+	}
 	return bFound;
 }
 
diff --git a/Engine/Plugins/Developer/GitSourceControl/Source/GitSourceControl/Private/GitSourceControlUtils.h b/Engine/Plugins/Developer/GitSourceControl/Source/GitSourceControl/Private/GitSourceControlUtils.h
--- a/Engine/Plugins/Developer/GitSourceControl/Source/GitSourceControl/Private/GitSourceControlUtils.h
+++ b/Engine/Plugins/Developer/GitSourceControl/Source/GitSourceControl/Private/GitSourceControlUtils.h
@@ -47,7 +47,7 @@ bool CheckGitAvailability(const FString& InPathToGitBinary);
 /**
  * Find the root of the Git repository, looking from the GameDir and upward in its parent directories
  * @param InPathToGameDir		The path to the Game Directory
- * @param OutRepositoryRoot		The path to the root directory of the Git repository if found
+ * @param OutRepositoryRoot		The path to the root directory of the Git repository if found, else the path to the GameDir
  * @returns true if the command succeeded and returned no errors
  */
 bool FindRootDirectory(const FString& InPathToGameDir, FString& OutRepositoryRoot);
