PR #1219: [4.8] (Git Plugin) added "*.opensdf" and "*.xcodeproj" to the ".gitignore" file (Contributed by [REDACTED_PERSON])

diff --git a/Engine/Plugins/Developer/GitSourceControl/Source/GitSourceControl/Private/SGitSourceControlSettings.cpp b/Engine/Plugins/Developer/GitSourceControl/Source/GitSourceControl/Private/SGitSourceControlSettings.cpp
--- a/Engine/Plugins/Developer/GitSourceControl/Source/GitSourceControl/Private/SGitSourceControlSettings.cpp
+++ b/Engine/Plugins/Developer/GitSourceControl/Source/GitSourceControl/Private/SGitSourceControlSettings.cpp
@@ -291,7 +291,7 @@ FReply SGitSourceControlSettings::OnClickedInitializeGitRepository()
 		{
 			// Create a standard ".gitignore" file with common patterns for a typical Blueprint & C++ project
 			const FString Filename = FString::Printf(TEXT("%s.gitignore"), *PathToGameDir);
-			const FString GitIgnoreContent = TEXT("Binaries\nDerivedDataCache\nIntermediate\nSaved\n*.sdf\n*.sln\n*.suo");
+			const FString GitIgnoreContent = TEXT("Binaries\nDerivedDataCache\nIntermediate\nSaved\n*.opensdf\n*.sdf\n*.sln\n*.suo\n*.xcodeproj");
 			if(FFileHelper::SaveStringToFile(GitIgnoreContent, *Filename, FFileHelper::EEncodingOptions::ForceUTF8WithoutBOM))
 			{
 				ProjectFiles.Add(TEXT(".gitignore"));
@@ -308,4 +308,4 @@ void SGitSourceControlSettings::OnCheckedCreateGitIgnore(ECheckBoxState NewCheck
 	bAutoCreateGitIgnore = (NewCheckedState == ECheckBoxState::Checked);
 }
 
-#undef LOCTEXT_NAMESPACE
\ No newline at end of file
+#undef LOCTEXT_NAMESPACE
