Fix crash executing syncronous source control commands with the GitHub plugin.

https://jira.ol.epicgames.net/browse/UE-18909

diff --git a/Engine/Plugins/Developer/GitSourceControl/Source/GitSourceControl/Private/GitSourceControlProvider.cpp b/Engine/Plugins/Developer/GitSourceControl/Source/GitSourceControl/Private/GitSourceControlProvider.cpp
--- a/Engine/Plugins/Developer/GitSourceControl/Source/GitSourceControl/Private/GitSourceControlProvider.cpp
+++ b/Engine/Plugins/Developer/GitSourceControl/Source/GitSourceControl/Private/GitSourceControlProvider.cpp
@@ -329,6 +329,12 @@ ECommandResult::Type FGitSourceControlProvider::ExecuteSynchronousCommand(FGitSo
 
 	// Delete the command now (asynchronous commands are deleted in the Tick() method)
 	check(!InCommand.bAutoDelete);
+
+	// ensure commands that are not auto deleted do not end up in the command queue
+	if ( CommandQueue.Contains( &InCommand ) ) 
+	{
+		CommandQueue.Remove( &InCommand );
+	}
 	delete &InCommand;
 
 	return Result;
