Fixed possible shadow variable.
Richard, It's not clear if this was an accidental shadow, or if you really meant to use the States member variable here.

diff --git a/Engine/Plugins/Developer/GitSourceControl/Source/GitSourceControl/Private/GitSourceControlOperations.cpp b/Engine/Plugins/Developer/GitSourceControl/Source/GitSourceControl/Private/GitSourceControlOperations.cpp
--- a/Engine/Plugins/Developer/GitSourceControl/Source/GitSourceControl/Private/GitSourceControlOperations.cpp
+++ b/Engine/Plugins/Developer/GitSourceControl/Source/GitSourceControl/Private/GitSourceControlOperations.cpp
@@ -71,9 +71,9 @@ bool FGitCheckInWorker::Execute(FGitSourceControlCommand& InCommand)
 			FGitSourceControlModule& GitSourceControl = FModuleManager::LoadModuleChecked<FGitSourceControlModule>("GitSourceControl");
 			FGitSourceControlProvider& Provider = GitSourceControl.GetProvider();
 
-			TArray<TSharedRef<ISourceControlState, ESPMode::ThreadSafe>> States;
-			Provider.GetState(InCommand.Files, States, EStateCacheUsage::Use);
-			for (const auto& State : States)
+			TArray<TSharedRef<ISourceControlState, ESPMode::ThreadSafe>> LocalStates;
+			Provider.GetState(InCommand.Files, LocalStates, EStateCacheUsage::Use);
+			for (const auto& State : LocalStates)
 			{
 				if (State->IsDeleted())
 				{
