Add support for the 64-bit git executable to the GitSourceControl plugin on windows.
From Pull Request: #1538

diff --git a/Engine/Plugins/Developer/GitSourceControl/Source/GitSourceControl/Private/GitSourceControlUtils.cpp b/Engine/Plugins/Developer/GitSourceControl/Source/GitSourceControl/Private/GitSourceControlUtils.cpp
--- a/Engine/Plugins/Developer/GitSourceControl/Source/GitSourceControl/Private/GitSourceControlUtils.cpp
+++ b/Engine/Plugins/Developer/GitSourceControl/Source/GitSourceControl/Private/GitSourceControlUtils.cpp
@@ -106,12 +106,19 @@ static bool RunCommandInternal(const FString& InCommand, const FString& InPathTo
 FString FindGitBinaryPath()
 {
 #if PLATFORM_WINDOWS
-	// 1) First of all, look into standard install directory
+	// 1) First of all, look into standard install directories
 	// NOTE using only "git" (or "git.exe") relying on the "PATH" envvar does not always work as expected, depending on the installation:
 	// If the PATH is set with "git/cmd" instead of "git/bin",
 	// "git.exe" launch "git/cmd/git.exe" that redirect to "git/bin/git.exe" and ExecProcess() is unable to catch its outputs streams.
-	FString GitBinaryPath(TEXT("C:/Program Files (x86)/Git/bin/git.exe"));
+	// First check the 64-bit program files directory:
+	FString GitBinaryPath(TEXT("C:/Program Files/Git/bin/git.exe"));
 	bool bFound = CheckGitAvailability(GitBinaryPath);
+	if(!bFound)
+	{
+		// otherwise check the 32-bit program files directory.
+		GitBinaryPath = TEXT("C:/Program Files (x86)/Git/bin/git.exe");
+		bFound = CheckGitAvailability(GitBinaryPath);
+	}
 
 	// 2) Else, look for the version of Git bundled with SmartGit "Installer with JRE"
 	if(!bFound)
