PR #4650: Git plugin: fix crash on commit error (Contributed by [REDACTED_PERSON])

diff --git a/Engine/Plugins/Developer/GitSourceControl/Source/GitSourceControl/Private/GitSourceControlOperations.cpp b/Engine/Plugins/Developer/GitSourceControl/Source/GitSourceControl/Private/GitSourceControlOperations.cpp
--- a/Engine/Plugins/Developer/GitSourceControl/Source/GitSourceControl/Private/GitSourceControlOperations.cpp
+++ b/Engine/Plugins/Developer/GitSourceControl/Source/GitSourceControl/Private/GitSourceControlOperations.cpp
@@ -97,7 +97,8 @@ bool FGitCheckInWorker::Execute(FGitSourceControlCommand& InCommand)
 			}
 
 			Operation->SetSuccessMessage(ParseCommitResults(InCommand.InfoMessages));
-			UE_LOG(LogSourceControl, Log, TEXT("commit successful: %s"), *InCommand.InfoMessages[0]);
+			const FString Message = (InCommand.InfoMessages.Num() > 0) ? InCommand.InfoMessages[0] : TEXT("");
+			UE_LOG(LogSourceControl, Log, TEXT("commit successful: %s"), *Message);
 		}
 	}
 
diff --git a/Engine/Plugins/Developer/GitSourceControl/Source/GitSourceControl/Private/GitSourceControlUtils.cpp b/Engine/Plugins/Developer/GitSourceControl/Source/GitSourceControl/Private/GitSourceControlUtils.cpp
--- a/Engine/Plugins/Developer/GitSourceControl/Source/GitSourceControl/Private/GitSourceControlUtils.cpp
+++ b/Engine/Plugins/Developer/GitSourceControl/Source/GitSourceControl/Private/GitSourceControlUtils.cpp
@@ -507,7 +507,7 @@ bool RunCommit(const FString& InPathToGitBinary, const FString& InRepositoryRoot
 	}
 	else
 	{
-		RunCommandInternal(TEXT("commit"), InPathToGitBinary, InRepositoryRoot, InParameters, InFiles, OutResults, OutErrorMessages);
+		bResult = RunCommandInternal(TEXT("commit"), InPathToGitBinary, InRepositoryRoot, InParameters, InFiles, OutResults, OutErrorMessages);
 	}
 
 	return bResult;
