PR #5106: Git Plugin: fix not finding directory using git worktree (Contributed by [REDACTED_PERSON])

diff --git a/Engine/Plugins/Developer/GitSourceControl/Source/GitSourceControl/Private/GitSourceControlUtils.cpp b/Engine/Plugins/Developer/GitSourceControl/Source/GitSourceControl/Private/GitSourceControlUtils.cpp
--- a/Engine/Plugins/Developer/GitSourceControl/Source/GitSourceControl/Private/GitSourceControlUtils.cpp
+++ b/Engine/Plugins/Developer/GitSourceControl/Source/GitSourceControl/Private/GitSourceControlUtils.cpp
@@ -334,9 +334,9 @@ bool FindRootDirectory(const FString& InPath, FString& OutRepositoryRoot)
 
 	while(!bFound && !OutRepositoryRoot.IsEmpty())
 	{
-		// Look for the ".git" subdirectory present at the root of every Git repository
+		// Look for the ".git" subdirectory (or file) present at the root of every Git repository
 		PathToGitSubdirectory = OutRepositoryRoot / TEXT(".git");
-		bFound = IFileManager::Get().DirectoryExists(*PathToGitSubdirectory);
+		bFound = IFileManager::Get().DirectoryExists(*PathToGitSubdirectory) || IFileManager::Get().FileExists(*PathToGitSubdirectory);
 		if(!bFound)
 		{
 			int32 LastSlashIndex;
