PR #4215: Add search path for git on mac (Contributed by [REDACTED_PERSON])
PR #4231: Fix for Git LFS with Git plugin to work on Mac (Contributed by [REDACTED_PERSON])

diff --git a/Engine/Plugins/Developer/GitSourceControl/Source/GitSourceControl/Private/GitSourceControlUtils.cpp b/Engine/Plugins/Developer/GitSourceControl/Source/GitSourceControl/Private/GitSourceControlUtils.cpp
--- a/Engine/Plugins/Developer/GitSourceControl/Source/GitSourceControl/Private/GitSourceControlUtils.cpp
+++ b/Engine/Plugins/Developer/GitSourceControl/Source/GitSourceControl/Private/GitSourceControlUtils.cpp
@@ -101,9 +101,33 @@ static bool RunCommandInternalRaw(const FString& InCommand, const FString& InPat
 #if UE_BUILD_DEBUG
 	UE_LOG(LogSourceControl, Log, TEXT("RunCommandInternalRaw: 'git %s'"), *LogableCommand);
 #endif
-	
-	FPlatformProcess::ExecProcess(*InPathToGitBinary, *FullCommand, &ReturnCode, &OutResults, &OutErrors);
-	
+
+	FString PathToGitOrEnvBinary = InPathToGitBinary;
+#if PLATFORM_MAC
+	// The Cocoa application does not inherit shell environment variables, so add the path expected to have git-lfs to PATH
+	FString PathEnv = FPlatformMisc::GetEnvironmentVariable(TEXT("PATH"));
+	FString GitInstallPath = FPaths::GetPath(InPathToGitBinary);
+
+	TArray<FString> PathArray;
+	PathEnv.ParseIntoArray(PathArray, FPlatformMisc::GetPathVarDelimiter());
+	bool bHasGitInstallPath = false;
+	for (auto Path : PathArray)
+	{
+		if (GitInstallPath.Equals(Path, ESearchCase::CaseSensitive))
+		{
+			bHasGitInstallPath = true;
+			break;
+		}
+	}
+
+	if (!bHasGitInstallPath)
+	{
+		PathToGitOrEnvBinary = FString("/usr/bin/env");
+		FullCommand = FString::Printf(TEXT("PATH=\"%s%s%s\" \"%s\" %s"), *GitInstallPath, FPlatformMisc::GetPathVarDelimiter(), *PathEnv, *InPathToGitBinary, *FullCommand);
+	}
+#endif
+	FPlatformProcess::ExecProcess(*PathToGitOrEnvBinary, *FullCommand, &ReturnCode, &OutResults, &OutErrors);
+
 #if UE_BUILD_DEBUG
 
 	if (OutResults.IsEmpty())
@@ -221,6 +245,85 @@ FString FindGitBinaryPath()
 		bFound = CheckGitAvailability(GitBinaryPath);
 	}
 
+#elif PLATFORM_MAC
+	// 1) First of all, look for the version of git provided by official git
+	FString GitBinaryPath = TEXT("/usr/local/git/bin/git");
+	bool bFound = CheckGitAvailability(GitBinaryPath);
+
+	// 2) Else, look for the version of git provided by Homebrew
+	if (!bFound)
+	{
+		GitBinaryPath = TEXT("/usr/local/bin/git");
+		bFound = CheckGitAvailability(GitBinaryPath);
+	}
+
+	// 3) Else, look for the version of git provided by MacPorts
+	if (!bFound)
+	{
+		GitBinaryPath = TEXT("/opt/local/bin/git");
+		bFound = CheckGitAvailability(GitBinaryPath);
+	}
+
+	// 4) Else, look for the version of git provided by Command Line Tools
+	if (!bFound)
+	{
+		GitBinaryPath = TEXT("/usr/bin/git");
+		bFound = CheckGitAvailability(GitBinaryPath);
+	}
+
+	{
+		SCOPED_AUTORELEASE_POOL;
+		NSWorkspace* SharedWorkspace = [NSWorkspace sharedWorkspace];
+
+		// 5) Else, look for the version of local_git provided by SmartGit
+		if (!bFound)
+		{
+			NSURL* AppURL = [SharedWorkspace URLForApplicationWithBundleIdentifier:@"com.syntevo.smartgit"];
+			if (AppURL != nullptr)
+			{
+				NSBundle* Bundle = [NSBundle bundleWithURL:AppURL];
+				GitBinaryPath = FString::Printf(TEXT("%s/git/bin/git"), *FString([Bundle resourcePath]));
+				bFound = CheckGitAvailability(GitBinaryPath);
+			}
+		}
+
+		// 6) Else, look for the version of local_git provided by SourceTree
+		if (!bFound)
+		{
+			NSURL* AppURL = [SharedWorkspace URLForApplicationWithBundleIdentifier:@"com.torusknot.SourceTreeNotMAS"];
+			if (AppURL != nullptr)
+			{
+				NSBundle* Bundle = [NSBundle bundleWithURL:AppURL];
+				GitBinaryPath = FString::Printf(TEXT("%s/git_local/bin/git"), *FString([Bundle resourcePath]));
+				bFound = CheckGitAvailability(GitBinaryPath);
+			}
+		}
+
+		// 7) Else, look for the version of local_git provided by GitHub Desktop
+		if (!bFound)
+		{
+			NSURL* AppURL = [SharedWorkspace URLForApplicationWithBundleIdentifier:@"com.github.GitHubClient"];
+			if (AppURL != nullptr)
+			{
+				NSBundle* Bundle = [NSBundle bundleWithURL:AppURL];
+				GitBinaryPath = FString::Printf(TEXT("%s/app/git/bin/git"), *FString([Bundle resourcePath]));
+				bFound = CheckGitAvailability(GitBinaryPath);
+			}
+		}
+
+		// 8) Else, look for the version of local_git provided by Tower2
+		if (!bFound)
+		{
+			NSURL* AppURL = [SharedWorkspace URLForApplicationWithBundleIdentifier:@"com.fournova.Tower2"];
+			if (AppURL != nullptr)
+			{
+				NSBundle* Bundle = [NSBundle bundleWithURL:AppURL];
+				GitBinaryPath = FString::Printf(TEXT("%s/git/bin/git"), *FString([Bundle resourcePath]));
+				bFound = CheckGitAvailability(GitBinaryPath);
+			}
+		}
+	}
+
 #else
 	FString GitBinaryPath = TEXT("/usr/bin/git");
 	bool bFound = CheckGitAvailability(GitBinaryPath);
