PR #5442: Git Plugin: add an option to create a README.md file at the time of the git repository creation (Contributed by [REDACTED_PERSON])

diff --git a/Engine/Plugins/Developer/GitSourceControl/Source/GitSourceControl/Private/SGitSourceControlSettings.cpp b/Engine/Plugins/Developer/GitSourceControl/Source/GitSourceControl/Private/SGitSourceControlSettings.cpp
--- a/Engine/Plugins/Developer/GitSourceControl/Source/GitSourceControl/Private/SGitSourceControlSettings.cpp
+++ b/Engine/Plugins/Developer/GitSourceControl/Source/GitSourceControl/Private/SGitSourceControlSettings.cpp
@@ -2,6 +2,7 @@
 
 #include "SGitSourceControlSettings.h"
 #include "Fonts/SlateFontInfo.h"
+#include "Misc/App.h"
 #include "Misc/Paths.h"
 #include "Misc/FileHelper.h"
 #include "Modules/ModuleManager.h"
@@ -27,9 +28,10 @@
 
 void SGitSourceControlSettings::Construct(const FArguments& InArgs)
 {
-	FSlateFontInfo Font = FEditorStyle::GetFontStyle(TEXT("SourceControl.LoginWindow.Font"));
+	const FSlateFontInfo Font = FEditorStyle::GetFontStyle(TEXT("SourceControl.LoginWindow.Font"));
 
 	bAutoCreateGitIgnore = true;
+	bAutoCreateReadme = true;
 	bAutoCreateGitAttributes = false;
 	bAutoInitialCommit = true;
 
@@ -42,6 +44,8 @@ void SGitSourceControlSettings::Construct(const FArguments& InArgs)
 	const FString FileFilterText = FString::Printf(TEXT("%s"), *FileFilterType.ToString());
 #endif
 
+	ReadmeContent = FText::FromString(FString(TEXT("# ")) + FApp::GetProjectName() + "\n\nDeveloped with Unreal Engine 4\n");
+
 	ChildSlot
 	[
 		SNew(SBorder)
@@ -228,6 +232,42 @@ void SGitSourceControlSettings::Construct(const FArguments& InArgs)
 					.Font(Font)
 				]
 			]
+			// Option to add a README.md file with custom content
+			+SVerticalBox::Slot()
+			.AutoHeight()
+			.Padding(2.0f)
+			.VAlign(VAlign_Center)
+			[
+				SNew(SHorizontalBox)
+				.Visibility(this, &SGitSourceControlSettings::CanInitializeGitRepository)
+				.ToolTipText(LOCTEXT("CreateReadme_Tooltip", "Add a README.md file"))
+				+SHorizontalBox::Slot()
+				.FillWidth(0.1f)
+				[
+					SNew(SCheckBox)
+					.IsChecked(ECheckBoxState::Checked)
+					.OnCheckStateChanged(this, &SGitSourceControlSettings::OnCheckedCreateReadme)
+				]
+				+SHorizontalBox::Slot()
+				.FillWidth(0.9f)
+				.VAlign(VAlign_Center)
+				[
+					SNew(STextBlock)
+					.Text(LOCTEXT("CreateReadme", "Add a basic README.md file"))
+					.Font(Font)
+				]
+				+SHorizontalBox::Slot()
+				.FillWidth(2.0f)
+				.Padding(2.0f)
+				[
+					SNew(SMultiLineEditableTextBox)
+					.Text(this, &SGitSourceControlSettings::GetReadmeContent)
+					.OnTextCommitted(this, &SGitSourceControlSettings::OnReadmeContentCommited)
+					.IsEnabled(this, &SGitSourceControlSettings::GetAutoCreateReadme)
+					.SelectAllTextWhenFocused(true)
+					.Font(Font)
+				]
+			]
 			// Option to add a proper .gitattributes file for Git LFS (false by default)
 			+SVerticalBox::Slot()
 			.AutoHeight()
@@ -418,12 +458,21 @@ FReply SGitSourceControlSettings::OnClickedInitializeGitRepository()
 				ProjectFiles.Add(GitIgnoreFilename);
 			}
 		}
-		if (bAutoCreateGitAttributes)
+		if(bAutoCreateReadme)
 		{
-			// 2.b. Synchronous (very quick) "lfs install" operation: needs only to be run once by user
+			// 2.b. Create a "README.md" file with a custom description
+			const FString ReadmeFilename = FPaths::Combine(FPaths::ProjectDir(), TEXT("README.md"));
+			if (FFileHelper::SaveStringToFile(ReadmeContent.ToString(), *ReadmeFilename, FFileHelper::EEncodingOptions::ForceUTF8WithoutBOM))
+			{
+				ProjectFiles.Add(ReadmeFilename);
+			}
+		}
+		if(bAutoCreateGitAttributes)
+		{
+			// 2.c. Synchronous (very quick) "lfs install" operation: needs only to be run once by user
 			GitSourceControlUtils::RunCommand(TEXT("lfs install"), PathToGitBinary, PathToProjectDir, TArray<FString>(), TArray<FString>(), InfoMessages, ErrorMessages);
 
-			// 2.c. Create a ".gitattributes" file to enable Git LFS (Large File System) for the whole "Content/" subdir
+			// 2.d. Create a ".gitattributes" file to enable Git LFS (Large File System) for the whole "Content/" subdir
 			const FString GitAttributesFilename = FPaths::Combine(FPaths::ProjectDir(), TEXT(".gitattributes"));
 			const FString GitAttributesContent = TEXT("Content/** filter=lfs diff=lfs merge=lfs -text\n");
 			if (FFileHelper::SaveStringToFile(GitAttributesContent, *GitAttributesFilename, FFileHelper::EEncodingOptions::ForceUTF8WithoutBOM))
@@ -544,6 +593,26 @@ void SGitSourceControlSettings::OnCheckedCreateGitIgnore(ECheckBoxState NewCheck
 	bAutoCreateGitIgnore = (NewCheckedState == ECheckBoxState::Checked);
 }
 
+void SGitSourceControlSettings::OnCheckedCreateReadme(ECheckBoxState NewCheckedState)
+{
+	bAutoCreateReadme = (NewCheckedState == ECheckBoxState::Checked);
+}
+
+bool SGitSourceControlSettings::GetAutoCreateReadme() const
+{
+	return bAutoCreateReadme;
+}
+
+void SGitSourceControlSettings::OnReadmeContentCommited(const FText& InText, ETextCommit::Type InCommitType)
+{
+	ReadmeContent = InText;
+}
+
+FText SGitSourceControlSettings::GetReadmeContent() const
+{
+	return ReadmeContent;
+}
+
 void SGitSourceControlSettings::OnCheckedCreateGitAttributes(ECheckBoxState NewCheckedState)
 {
 	bAutoCreateGitAttributes = (NewCheckedState == ECheckBoxState::Checked);
diff --git a/Engine/Plugins/Developer/GitSourceControl/Source/GitSourceControl/Private/SGitSourceControlSettings.h b/Engine/Plugins/Developer/GitSourceControl/Source/GitSourceControl/Private/SGitSourceControlSettings.h
--- a/Engine/Plugins/Developer/GitSourceControl/Source/GitSourceControl/Private/SGitSourceControlSettings.h
+++ b/Engine/Plugins/Developer/GitSourceControl/Source/GitSourceControl/Private/SGitSourceControlSettings.h
@@ -29,10 +29,8 @@ public:
 
 private:
 
-	/** Delegate to get binary path from settings */
+	/** Delegates to get Git binary path from/to settings */
 	FString GetBinaryPathString() const;
-
-	/** Delegate to commit repository text to settings */
 	void OnBinaryPathPicked(const FString & PickedPath) const;
 
 	/** Delegate to get repository root, user name and email from provider */
@@ -49,6 +47,14 @@ private:
 	void OnCheckedCreateGitIgnore(ECheckBoxState NewCheckedState);
 	bool bAutoCreateGitIgnore;
 
+	/** Delegates to create a README.md file */
+	void OnCheckedCreateReadme(ECheckBoxState NewCheckedState);
+	bool GetAutoCreateReadme() const;
+	bool bAutoCreateReadme;
+	void OnReadmeContentCommited(const FText& InText, ETextCommit::Type InCommitType);
+	FText GetReadmeContent() const;
+	FText ReadmeContent;
+
 	void OnCheckedCreateGitAttributes(ECheckBoxState NewCheckedState);
 	bool bAutoCreateGitAttributes;
 
