PR #8091: Fixed Default Commit Message, Reflects Current Version (Contributed by [REDACTED_PERSON])

diff --git a/Engine/Plugins/Developer/GitSourceControl/Source/GitSourceControl/Private/SGitSourceControlSettings.cpp b/Engine/Plugins/Developer/GitSourceControl/Source/GitSourceControl/Private/SGitSourceControlSettings.cpp
--- a/Engine/Plugins/Developer/GitSourceControl/Source/GitSourceControl/Private/SGitSourceControlSettings.cpp
+++ b/Engine/Plugins/Developer/GitSourceControl/Source/GitSourceControl/Private/SGitSourceControlSettings.cpp
@@ -41,7 +41,7 @@ void SGitSourceControlSettings::Construct(const FArguments& InArgs)
 	const FString FileFilterText = FString::Printf(TEXT("%s"), *FileFilterType.ToString());
 #endif
 
-	ReadmeContent = FText::FromString(FString(TEXT("# ")) + FApp::GetProjectName() + "\n\nDeveloped with Unreal Engine 4\n");
+	ReadmeContent = FText::FromString(FString(TEXT("# ")) + FApp::GetProjectName() + "\n\nDeveloped with Unreal Engine 5\n");
 
 	ChildSlot
 	[
