Fixed issue where diffing files could fail when using Git LFS

diff --git a/Engine/Plugins/Developer/GitSourceControl/Source/GitSourceControl/Private/GitSourceControlUtils.cpp b/Engine/Plugins/Developer/GitSourceControl/Source/GitSourceControl/Private/GitSourceControlUtils.cpp
--- a/Engine/Plugins/Developer/GitSourceControl/Source/GitSourceControl/Private/GitSourceControlUtils.cpp
+++ b/Engine/Plugins/Developer/GitSourceControl/Source/GitSourceControl/Private/GitSourceControlUtils.cpp
@@ -1005,7 +1005,7 @@ bool RunDumpToFile(const FString& InPathToGitBinary, const FString& InRepository
 
 	verify(FPlatformProcess::CreatePipe(PipeRead, PipeWrite));
 
-	FProcHandle ProcessHandle = FPlatformProcess::CreateProc(*InPathToGitBinary, *FullCommand, bLaunchDetached, bLaunchHidden, bLaunchReallyHidden, nullptr, 0, *InRepositoryRoot, PipeWrite);
+	FProcHandle ProcessHandle = FPlatformProcess::CreateProc(*InPathToGitBinary, *FullCommand, bLaunchDetached, bLaunchHidden, bLaunchReallyHidden, nullptr, 0, *InRepositoryRoot, PipeWrite, nullptr, nullptr);
 	if(ProcessHandle.IsValid())
 	{
 		FPlatformProcess::Sleep(0.01);
