Give the option to disable being able to diff files in the Source Control Submit Widget.
Add the missing default impelementation for FGitSourceControlProvider

diff --git a/Engine/Plugins/Developer/GitSourceControl/Source/GitSourceControl/Private/GitSourceControlProvider.cpp b/Engine/Plugins/Developer/GitSourceControl/Source/GitSourceControl/Private/GitSourceControlProvider.cpp
--- a/Engine/Plugins/Developer/GitSourceControl/Source/GitSourceControl/Private/GitSourceControlProvider.cpp
+++ b/Engine/Plugins/Developer/GitSourceControl/Source/GitSourceControl/Private/GitSourceControlProvider.cpp
@@ -272,6 +272,11 @@ bool FGitSourceControlProvider::UsesFileRevisions() const
 	return false;
 }
 
+bool FGitSourceControlProvider::AllowsDiffAgainstDepot() const
+{
+	return true;
+}
+
 TOptional<bool> FGitSourceControlProvider::IsAtLatestRevision() const
 {
 	return TOptional<bool>();
