Revision Control: Update conflicted icon for git

diff --git a/Engine/Plugins/Developer/GitSourceControl/Source/GitSourceControl/Private/GitSourceControlState.cpp b/Engine/Plugins/Developer/GitSourceControl/Source/GitSourceControl/Private/GitSourceControlState.cpp
--- a/Engine/Plugins/Developer/GitSourceControl/Source/GitSourceControl/Private/GitSourceControlState.cpp
+++ b/Engine/Plugins/Developer/GitSourceControl/Source/GitSourceControl/Private/GitSourceControlState.cpp
@@ -77,7 +77,7 @@ FSlateIcon FGitSourceControlState::GetIcon() const
 	case EWorkingCopyState::Missing:
 		return FSlateIcon(FRevisionControlStyleManager::GetStyleSetName(), "RevisionControl.MarkedForDelete");
 	case EWorkingCopyState::Conflicted:
-		return FSlateIcon(FRevisionControlStyleManager::GetStyleSetName(), "RevisionControl.NotAtHeadRevision");
+		return FSlateIcon(FRevisionControlStyleManager::GetStyleSetName(), "RevisionControl.Conflicted");
 	case EWorkingCopyState::NotControlled:
 		return FSlateIcon(FRevisionControlStyleManager::GetStyleSetName(), "RevisionControl.NotInDepot");
 	case EWorkingCopyState::Unknown:
