Truncation warning fixes for GitSourceControl module

#rnx

diff --git a/Engine/Plugins/Developer/GitSourceControl/Source/GitSourceControl/GitSourceControl.Build.cs b/Engine/Plugins/Developer/GitSourceControl/Source/GitSourceControl/GitSourceControl.Build.cs
--- a/Engine/Plugins/Developer/GitSourceControl/Source/GitSourceControl/GitSourceControl.Build.cs
+++ b/Engine/Plugins/Developer/GitSourceControl/Source/GitSourceControl/GitSourceControl.Build.cs
@@ -24,5 +24,7 @@ public class GitSourceControl : ModuleRules
 			PrivateDependencyModuleNames.Add("EditorFramework");
 			PrivateDependencyModuleNames.Add("UnrealEd");
 		}
+
+		UnsafeTypeCastWarningLevel = WarningLevel.Error;
 	}
 }
diff --git a/Engine/Plugins/Developer/GitSourceControl/Source/GitSourceControl/Private/GitSourceControlUtils.cpp b/Engine/Plugins/Developer/GitSourceControl/Source/GitSourceControl/Private/GitSourceControlUtils.cpp
--- a/Engine/Plugins/Developer/GitSourceControl/Source/GitSourceControl/Private/GitSourceControlUtils.cpp
+++ b/Engine/Plugins/Developer/GitSourceControl/Source/GitSourceControl/Private/GitSourceControlUtils.cpp
@@ -1021,7 +1021,7 @@ bool RunDumpToFile(const FString& InPathToGitBinary, const FString& InRepository
 	FProcHandle ProcessHandle = FPlatformProcess::CreateProc(*InPathToGitBinary, *FullCommand, bLaunchDetached, bLaunchHidden, bLaunchReallyHidden, nullptr, 0, *InRepositoryRoot, PipeWrite, nullptr, nullptr);
 	if(ProcessHandle.IsValid())
 	{
-		FPlatformProcess::Sleep(0.01);
+		FPlatformProcess::Sleep(0.01f);
 
 		TArray<uint8> BinaryFileContent;
 		while(FPlatformProcess::IsProcRunning(ProcessHandle))
diff --git a/Engine/Plugins/Developer/GitSourceControl/Source/GitSourceControl/Private/SGitSourceControlSettings.cpp b/Engine/Plugins/Developer/GitSourceControl/Source/GitSourceControl/Private/SGitSourceControlSettings.cpp
--- a/Engine/Plugins/Developer/GitSourceControl/Source/GitSourceControl/Private/SGitSourceControlSettings.cpp
+++ b/Engine/Plugins/Developer/GitSourceControl/Source/GitSourceControl/Private/SGitSourceControlSettings.cpp
@@ -304,7 +304,7 @@ void SGitSourceControlSettings::Construct(const FArguments& InArgs)
 				.ToolTipText(LOCTEXT("GitInitRepository_Tooltip", "Initialize current project as a new Git repository"))
 				.OnClicked(this, &SGitSourceControlSettings::OnClickedInitializeGitRepository)
 				.HAlign(HAlign_Center)
-				.ContentPadding(6)
+				.ContentPadding(6.0f)
 			]
 		]
 	];
