[BugFix] Typo in FGitSourceControlState::FindHistoryRevision that caused the function to fail with valid revision numbers

diff --git a/Engine/Plugins/Developer/GitSourceControl/Source/GitSourceControl/Private/GitSourceControlState.cpp b/Engine/Plugins/Developer/GitSourceControl/Source/GitSourceControl/Private/GitSourceControlState.cpp
--- a/Engine/Plugins/Developer/GitSourceControl/Source/GitSourceControl/Private/GitSourceControlState.cpp
+++ b/Engine/Plugins/Developer/GitSourceControl/Source/GitSourceControl/Private/GitSourceControlState.cpp
@@ -40,7 +40,7 @@ TSharedPtr<class ISourceControlRevision, ESPMode::ThreadSafe> FGitSourceControlS
 TSharedPtr<class ISourceControlRevision, ESPMode::ThreadSafe> FGitSourceControlState::FindHistoryRevision(const FString& InRevision) const
 {
 	// short hash must be >= 7 characters to have a reasonable probability of finding the correct revision
-	if (!ensure(InRevision.Len() < 7))
+	if (!ensure(InRevision.Len() >= 7))
 	{
 		return nullptr;
 	}
