Unshelved from pending changelist '30117699':

PR #11193: Fix FindHistoryRevision for git source control

diff --git a/Engine/Plugins/Developer/GitSourceControl/Source/GitSourceControl/Private/GitSourceControlState.cpp b/Engine/Plugins/Developer/GitSourceControl/Source/GitSourceControl/Private/GitSourceControlState.cpp
--- a/Engine/Plugins/Developer/GitSourceControl/Source/GitSourceControl/Private/GitSourceControlState.cpp
+++ b/Engine/Plugins/Developer/GitSourceControl/Source/GitSourceControl/Private/GitSourceControlState.cpp
@@ -48,9 +48,9 @@ TSharedPtr<class ISourceControlRevision, ESPMode::ThreadSafe> FGitSourceControlS
 	for(const auto& Revision : History)
 	{
 		// support for short hashes
-		const int32 Len = FMath::Min(Revision->FileHash.Len(), InRevision.Len());
+		const int32 Len = FMath::Min(Revision->CommitId.Len(), InRevision.Len());
 		
-		if(Revision->FileHash.Left(Len) == InRevision.Left(Len))
+		if(Revision->CommitId.Left(Len) == InRevision.Left(Len))
 		{
 			return Revision;
 		}
