GitSourceControl:
- Only accept git named executables from its settings.

diff --git a/Engine/Plugins/Developer/GitSourceControl/Source/GitSourceControl/Private/GitSourceControlUtils.cpp b/Engine/Plugins/Developer/GitSourceControl/Source/GitSourceControl/Private/GitSourceControlUtils.cpp
--- a/Engine/Plugins/Developer/GitSourceControl/Source/GitSourceControl/Private/GitSourceControlUtils.cpp
+++ b/Engine/Plugins/Developer/GitSourceControl/Source/GitSourceControl/Private/GitSourceControlUtils.cpp
@@ -342,6 +342,18 @@ FString FindGitBinaryPath()
 
 bool CheckGitAvailability(const FString& InPathToGitBinary, FGitVersion *OutVersion)
 {
+#if PLATFORM_WINDOWS
+	const TCHAR* GitExecutableSuffix = TEXT("git.exe");
+#else
+	const TCHAR* GitExecutableSuffix = TEXT("git");
+#endif
+
+	// Make sure we don't run arbitrary executables
+	if ( !InPathToGitBinary.EndsWith(GitExecutableSuffix) )
+	{
+		return false;
+	}
+
 	FString InfoMessages;
 	FString ErrorMessages;
 	bool bGitAvailable = RunCommandInternalRaw(TEXT("version"), InPathToGitBinary, FString(), TArray<FString>(), TArray<FString>(), InfoMessages, ErrorMessages);
