Fix for `RunDumpToFile` to finish reading the pipe

The current code of `RunDumpToFile` only reads the pipe once when the process stops running. It is possible that not all the data is read in one call, and that multiple calls to `FPlatformProcess::ReadPipeToArray` are necessary. This has happened on Mac where `FMacPlatformProcess::ReadPipeToArray` only reads 32768 bytes.

#12828

diff --git a/Engine/Plugins/Developer/GitSourceControl/Source/GitSourceControl/Private/GitSourceControlUtils.cpp b/Engine/Plugins/Developer/GitSourceControl/Source/GitSourceControl/Private/GitSourceControlUtils.cpp
--- a/Engine/Plugins/Developer/GitSourceControl/Source/GitSourceControl/Private/GitSourceControlUtils.cpp
+++ b/Engine/Plugins/Developer/GitSourceControl/Source/GitSourceControl/Private/GitSourceControlUtils.cpp
@@ -1039,15 +1039,13 @@ bool RunDumpToFile(const FString& InPathToGitBinary, const FString& InRepository
 		while(FPlatformProcess::IsProcRunning(ProcessHandle))
 		{
 			TArray<uint8> BinaryData;
-			FPlatformProcess::ReadPipeToArray(PipeRead, BinaryData);
-			if(BinaryData.Num() > 0)
+			if(FPlatformProcess::ReadPipeToArray(PipeRead, BinaryData))
 			{
 				BinaryFileContent.Append(MoveTemp(BinaryData));
 			}
 		}
 		TArray<uint8> BinaryData;
-		FPlatformProcess::ReadPipeToArray(PipeRead, BinaryData);
-		if(BinaryData.Num() > 0)
+		while(FPlatformProcess::ReadPipeToArray(PipeRead, BinaryData))
 		{
 			BinaryFileContent.Append(MoveTemp(BinaryData));
 		}
